Cherrypicking Mirage code base (disabled) from Dev-Cooker

- The mirage specific code is disabled behind the define UE_USE_VIRTUALBULKDATA, this means that some code paths in Texture/Mesh are much more complex than they need to be as we support both old and new paths. Once the system has been turned on and confirmed to cause no issues then this will be stripped out.
- SavePackageUtilities.cpp, SavePackage.cpp and SavePackage2.cpp are editgrates rather than integrations as those files have changes in DevCooker that we don't want to bring over immediately.
- Also includes a prototype system for storing bulkdata in a sidecar file in the workspace domain rather than in the .uasset/.umap file which although has been discontinued as part of mirage, will have applications for future work for non-virtualized projects and/or text based assets.

#tests Cooking and running ShooterGame/Frosty and other sample programs using megascan assets
#rnx

diff --git a/Engine/Config/BaseEngine.ini b/Engine/Config/BaseEngine.ini
--- a/Engine/Config/BaseEngine.ini
+++ b/Engine/Config/BaseEngine.ini
@@ -1045,9 +1045,35 @@ TaskThreadStackSize=393216
 PoolThreadStackSize=131072
 DefaultThreadStackSize=1048576
 
+[Core.System.Experimental]
+EnablePackageSidecarSaving=False
+
+; Options for the content virtualization system
 [Core.ContentVirtualization]
 EnablePushToBackend=False
 MinPayloadLength=0
+BackendGraph=ContentVirtualizationBackendGraph_None
+
+; Options to help debug/profile problems with the content virtualization system
+[Core.ContentVirtualizationDebugOptions]
+ForceSingleThreaded=False
+FailPayloadPullOperations=False
+ValidateAfterPushOperation=False
+
+; A backend graph with no backends
+[ContentVirtualizationBackendGraph_None]
+Hierarchy=()
+
+; A simple backend graph that will store virtualized data locally, suitable for testing.
+[ContentVirtualizationBackendGraph_LocalCache]
+Hierarchy=(Entry=LocalCache)
+LocalCache=(Type=FileSystem, Path="%GAMESAVEDDIR%VirtualizedPayloads/")
+
+; A backend graph that will connect to a locally hosted Jupiter service along with a local file cache, suitable for testing.
+[ContentVirtualizationBackendGraph_JupiterLocalHost]
+Hierarchy=(Entry=LocalCache, Entry=JupiterCache)
+LocalCache=(Type=FileSystem, Path="%GAMESAVEDDIR%VirtualizedPayloads/")
+JupiterCache=(Type=Jupiter, Host="http://localhost", Namespace="ue.mirage", ChunkSize=524288)
 
 [/Script/Engine.StreamingSettings]
 s.MinBulkDataSizeForAsyncLoading=131072
diff --git a/Engine/Source/Developer/Virtualization/Private/VirtualizationFileBackend.cpp b/Engine/Source/Developer/Virtualization/Private/VirtualizationFileBackend.cpp
new file mode 100644
--- /dev/null
+++ b/Engine/Source/Developer/Virtualization/Private/VirtualizationFileBackend.cpp
@@ -0,0 +1,154 @@
+// Copyright Epic Games, Inc. All Rights Reserved.
+
+#include "Virtualization/IVirtualizationBackend.h"
+
+#include "Containers/StringView.h"
+#include "HAL/FileManager.h"
+#include "Misc/Parse.h"
+#include "Misc/Paths.h"
+#include "Virtualization/VirtualizationManager.h"
+
+namespace UE::Virtualization
+{
+
+/**
+ * A basic backend based on the file system. This can be used to access/store virtualization
+ * data either on a local disk or a network share. It is intended to be used as a caching system
+ * to speed up operations (running a local cache or a shared cache for a site) rather than as the
+ * proper backend solution.
+ * 
+ * Ini file setup:
+ * 'Name'=(Type=FileSystem, Path="XXX")
+ * Where 'Name' is the backend name in the hierarchy and 'XXX' if the path to the directory where
+ * you want to files to be stored.
+ */
+class FFileSystemBackend : public IVirtualizationBackend
+{
+public:
+	FFileSystemBackend(FStringView ConfigName)
+		: IVirtualizationBackend(EOperations::Both)
+	{
+		Name = TEXT("FFileSystemBackend - ");
+		Name.Append(ConfigName);
+	}
+
+	virtual ~FFileSystemBackend() = default;
+
+private:
+
+	virtual bool Initialize(const FString& ConfigEntry) override
+	{
+		if (!FParse::Value(*ConfigEntry, TEXT("Path="), RootDirectory))
+		{
+			UE_LOG(LogVirtualization, Error, TEXT("[%s] 'Path=' not found in the config file"), *GetDebugString());
+			return false;
+		}
+
+		FPaths::NormalizeDirectoryName(RootDirectory);
+
+		if (RootDirectory.IsEmpty())
+		{
+			UE_LOG(LogVirtualization, Error, TEXT("[%s] Config file entry 'Path=' was empty"), *GetDebugString());
+			return false;
+		}
+
+		// TODO: Validate that the given path is usable?
+
+		UE_LOG(LogVirtualization, Log, TEXT("[%s] Using path: '%s'"), *GetDebugString(), *RootDirectory);
+
+		return true;
+	}
+
+	virtual bool PushData(const FPayloadId& Id, const FCompressedBuffer& Payload) override
+	{
+		TRACE_CPUPROFILER_EVENT_SCOPE(FFileSystemBackend::PushData);
+
+		if (DoesExist(Id))
+		{
+			UE_LOG(LogVirtualization, Verbose, TEXT("[%s] Already has a copy of the payload '%s'."), *GetDebugString(), *Id.ToString());
+			return true;
+		}
+
+		const FString FilePath = CreateFilePath(Id);
+
+		{
+			// TODO: Should we write to a temp file and then move it once it has written?
+			TUniquePtr<FArchive> FileAr(IFileManager::Get().CreateFileWriter(*FilePath));
+			if (FileAr == nullptr)
+			{
+				UE_LOG(LogVirtualization, Error, TEXT("[%s] Failed to push payload '%s' to '%s'"), *GetDebugString(), *Id.ToString(), *FilePath);
+				return false;
+			}
+
+			for (const FSharedBuffer& Buffer : Payload.GetCompressed().GetSegments())
+			{
+				// Const cast because FArchive requires a non-const pointer!
+				FileAr->Serialize(const_cast<void*>(Buffer.GetData()), static_cast<int64>(Buffer.GetSize()));
+			}
+		}
+
+		return true;
+	}
+
+	virtual FCompressedBuffer PullData(const FPayloadId& Id) override
+	{
+		TRACE_CPUPROFILER_EVENT_SCOPE(FFileSystemBackend::PullData);
+
+		const FString FilePath = CreateFilePath(Id);
+
+		// TODO: Should we allow the error severity to be configured via ini or just not report this case at all?
+		if (!IFileManager::Get().FileExists(*FilePath))
+		{
+			UE_LOG(LogVirtualization, Verbose, TEXT("[%s] Does not contain the payload '%s'"), *GetDebugString(), *Id.ToString());
+			return FCompressedBuffer();
+		}
+
+		TUniquePtr<FArchive> FileAr(IFileManager::Get().CreateFileReader(*FilePath));
+		if (FileAr == nullptr)
+		{
+			const uint32 SystemError = FPlatformMisc::GetLastError();
+			// If we have a system error we can give a more informative error message but don't output it if the error is zero as 
+			// this can lead to very confusing error messages.
+			if (SystemError != 0)
+			{
+				TCHAR SystemErrorMsg[2048] = { 0 };
+				FPlatformMisc::GetSystemErrorMessage(SystemErrorMsg, sizeof(SystemErrorMsg), SystemError);
+				UE_LOG(LogVirtualization, Error, TEXT("[%s] Failed to load payload '%s' file '%s' due to system error: '%s' (%d))"), *GetDebugString(), *Id.ToString(), *FilePath, SystemErrorMsg, SystemError);
+			}
+			else
+			{
+				UE_LOG(LogVirtualization, Error, TEXT("[%s] Failed to load payload '%s' from '%s' (reason unknown)"), *GetDebugString(), *Id.ToString(), *FilePath);
+			}
+		
+			return FCompressedBuffer();
+		}
+
+		return FCompressedBuffer::FromCompressed(*FileAr);
+	}
+
+	bool DoesExist(const FPayloadId& Id)
+	{
+		TRACE_CPUPROFILER_EVENT_SCOPE(FFileSystemBackend::DoesExist);
+
+		const FString FilePath = CreateFilePath(Id);
+		return IFileManager::Get().FileExists(*FilePath);
+	}
+
+	virtual FString GetDebugString() const override
+	{
+		return Name;
+	}
+
+	FString CreateFilePath(const FPayloadId& PayloadId)
+	{
+		return FString::Printf(TEXT("%s/%s.payload"), *RootDirectory, *PayloadId.ToString());
+	}
+
+private:
+	FString Name;
+	FString RootDirectory;
+};
+
+UE_REGISTER_VIRTUALIZATION_BACKEND_FACTORY(FFileSystemBackend, FileSystem);
+
+} // namespace UE::Virtualization
diff --git a/Engine/Source/Developer/Virtualization/Private/VirtualizationJupiterBackend.cpp b/Engine/Source/Developer/Virtualization/Private/VirtualizationJupiterBackend.cpp
new file mode 100644
--- /dev/null
+++ b/Engine/Source/Developer/Virtualization/Private/VirtualizationJupiterBackend.cpp
@@ -0,0 +1,2066 @@
+// Copyright Epic Games, Inc. All Rights Reserved.
+
+#include "Virtualization/IVirtualizationBackend.h"
+
+// TODO: Our libcurl implementation does not currently support MacOS
+// (not registering this will cause a fatal log error if the backend is actually used)
+#if PLATFORM_WINDOWS
+#define WITH_MIRAGE_JUPITER_BACKEND 1
+#else
+#define WITH_MIRAGE_JUPITER_BACKEND 0
+#endif //PLATFORM_WINDOWS
+
+#if WITH_MIRAGE_JUPITER_BACKEND
+
+#include "Async/TaskGraphInterfaces.h"
+#include "Virtualization/PayloadId.h"
+#include "Compression/CompressedBuffer.h"
+#include "Containers/StringConv.h"
+#include "Containers/Ticker.h"
+#include "Logging/LogMacros.h"
+#include "Misc/FileHelper.h"
+#include "Misc/Parse.h"
+#include "Misc/ScopeLock.h"
+#include "Misc/StringBuilder.h"
+#include "Serialization/JsonReader.h"
+#include "Serialization/JsonSerializer.h"
+#include "Serialization/JsonSerializerMacros.h"
+#include "Serialization/MemoryReader.h"
+#include "Serialization/MemoryWriter.h"
+
+// The following inlcudes are all for the Utility namespace
+#if PLATFORM_WINDOWS || PLATFORM_HOLOLENS
+	#include "Windows/WindowsHWrapper.h"
+	#include "Windows/AllowWindowsPlatformTypes.h"
+#endif
+
+#include "curl/curl.h"
+
+#if PLATFORM_WINDOWS || PLATFORM_HOLOLENS
+	#include "Windows/HideWindowsPlatformTypes.h"
+#endif
+
+#include "Containers/StaticArray.h"
+#include "Dom/JsonObject.h"
+#include "HAL/PlatformProcess.h"
+#include "Misc/ScopeRWLock.h"
+#include "Misc/SecureHash.h"
+#include "ProfilingDebugging/CountersTrace.h"
+
+#if WITH_SSL
+	#include "Ssl.h"
+#define UI UI_ST // Work around for name clash with the UI namespace in CoreUObject
+	#include <openssl/ssl.h>
+#undef UI
+#endif //WITH_SSL
+
+/** 
+ * The min version that the Jupiter service must be in order for us to connect.
+ * Only increase the minimum required version to ensure that specific features
+ * are present that the code cannot run without.
+ */
+#define UE_MIRAGE_JUPITER_MIN_MAJOR_VER 0
+#define UE_MIRAGE_JUPITER_MIN_MINOR_VER 27
+#define UE_MIRAGE_JUPITER_MIN_PATCH_VER 5
+
+/** When enabled we will only attempt to upload a payload if Jupiter claims to not already have it.
+	Disabling this is only intended for debug purposes. */
+#define UE_CHECK_FOR_EXISTING_PAYLOADS 0
+
+/** When enabled we will only attempt to upload a chunk if Jupiter claims to not already have it. 
+	In practice this doesn't really speed up the workflow hence being disabled. */
+#define UE_CHECK_FOR_EXISTING_CHUNKS 0
+
+/** When enabled we will attempt to PUT/GET many chunks at the same time asynchronously to improve throughput.
+	Disabling this is only intended for debug purposes as it is significantly slower. */
+#define UE_ENABLE_ASYNC_CHUNK_ACCESS 1
+
+TRACE_DECLARE_INT_COUNTER(Jupiter_WaitOnRequestPool, TEXT("Jupiter Wait On Request Pool"));
+
+namespace UE
+{
+// Code in the following name space has been taken from HttpDerivedDataBackend.cpp
+// once done we should consider what functionality is used by both implementations 
+// and what might be used by future implementations then decide what should be moved 
+// to a shared module
+namespace Utility
+{
+#define UE_MIRAGE_REQUEST_POOL_WAIT_INTERVAL 0.01f
+#define UE_MIRAGE_REQUEST_TIMEOUT_SECONDS 30L
+#define UE_MIRAGE_REQUEST_TIMOUT_ENABLED 1
+#define UE_MIRAGE_DEBUG 0
+#define UE_MIRAGE_REQUEST_POOL_SIZE 32
+#define UE_MIRAGE_MAX_FAILED_LOGIN_ATTEMPTS 16
+#define UE_MIRAGE_MAX_ATTEMPTS 4
+#define UE_MIRAGE_MAX_BUFFER_RESERVE 104857600u
+
+#if WITH_SSL
+static CURLcode sslctx_function(CURL* curl, void* sslctx, void* parm);
+#endif //WITH_SSL
+
+/**
+* Encapsulation for access token shared by all requests.
+*/
+struct FAccessToken
+{
+public:
+	FAccessToken() = default;
+	FString GetHeader() const;
+	void SetHeader(const TCHAR*);
+	uint32 GetSerial() const;
+
+private:
+	mutable FRWLock	Lock;
+	FString			Token;
+	uint32			Serial;
+};
+
+FString FAccessToken::GetHeader() const
+{
+	FReadScopeLock _(Lock);
+
+	return FString::Printf(TEXT("Authorization: Bearer %s"), *Token);
+}
+
+void FAccessToken::SetHeader(const TCHAR* InToken)
+{
+	FWriteScopeLock _(Lock);
+
+	Token = InToken;
+	Serial++;
+}
+
+uint32 FAccessToken::GetSerial() const
+{
+	return Serial;
+}
+
+/**
+* Minimal HTTP request type wrapping CURL without the need for managers. This request
+* is written to allow reuse of request objects, in order to allow connections to be reused.
+*
+* CURL has a global library initialization (curl_global_init). We rely on this happening in 
+* the Online/HTTP library which is a dependency on this module.
+*/
+class FRequest
+{
+public:
+	/**
+	* Supported request verbs
+	*/
+	enum RequestVerb
+	{
+		Get,
+		Put,
+		PutJson,
+		Post,
+		PostJson,
+		Delete,
+		Head
+	};
+
+	/**
+	* Convenience result type interpreted from HTTP response code.
+	*/
+	enum Result
+	{
+		Success,
+		Failed
+	};
+
+	enum EContentType
+	{
+		None = 0,
+		OctetStream,
+		UrlEncoded,
+		Json,
+		Xml
+	};
+
+	FRequest(const TCHAR* InDomain, FAccessToken* InAuthorizationToken, bool bInLogErrors)
+		: bLogErrors(bInLogErrors)
+		, Domain(InDomain)
+		, AuthorizationToken(InAuthorizationToken)
+	{
+		Curl = curl_easy_init();
+		Reset();
+	}
+
+	~FRequest()
+	{
+		curl_easy_cleanup(Curl);
+	}
+
+	/**
+	* Resets all options on the request except those that should always be set.
+	*/
+	void Reset()
+	{
+		Headers.Reset();
+		ResponseHeader.Reset();
+		ResponseBuffer.Reset();
+		ResponseCode = 0;
+		ReadDataView = TArrayView<const uint8>();
+		WriteDataView.Reset();
+		WriteDataBufferPtr = nullptr;
+		WriteHeaderBufferPtr = nullptr;
+		BytesSent = 0;
+		BytesReceived = 0;
+		CurlResult = CURL_LAST;
+
+		curl_easy_reset(Curl);
+
+		// Options that are always set for all connections.
+#if UE_MIRAGE_REQUEST_TIMOUT_ENABLED
+		curl_easy_setopt(Curl, CURLOPT_CONNECTTIMEOUT, UE_MIRAGE_REQUEST_TIMEOUT_SECONDS);
+#endif // UE_MIRAGE_REQUEST_TIMOUT_ENABLED
+		curl_easy_setopt(Curl, CURLOPT_FOLLOWLOCATION, 1L);
+		curl_easy_setopt(Curl, CURLOPT_NOSIGNAL, 1L);
+		// SSL options
+		curl_easy_setopt(Curl, CURLOPT_USE_SSL, CURLUSESSL_ALL);
+		curl_easy_setopt(Curl, CURLOPT_SSL_VERIFYPEER, 1);
+		curl_easy_setopt(Curl, CURLOPT_SSL_VERIFYHOST, 1);
+		curl_easy_setopt(Curl, CURLOPT_SSLCERTTYPE, "PEM");
+		// Response functions
+		curl_easy_setopt(Curl, CURLOPT_HEADERDATA, this);
+		curl_easy_setopt(Curl, CURLOPT_HEADERFUNCTION, &FRequest::StaticWriteHeaderFn);
+		curl_easy_setopt(Curl, CURLOPT_WRITEDATA, this);
+		curl_easy_setopt(Curl, CURLOPT_WRITEFUNCTION, StaticWriteBodyFn);
+		// SSL certification verification
+		curl_easy_setopt(Curl, CURLOPT_CAINFO, nullptr);
+		// TODO: What do we do if we do not have SSL?
+#if WITH_SSL
+		curl_easy_setopt(Curl, CURLOPT_SSL_CTX_FUNCTION, *sslctx_function);
+#endif // WITH_SSL
+		curl_easy_setopt(Curl, CURLOPT_SSL_CTX_DATA, this);
+		// Allow compressed data
+		curl_easy_setopt(Curl, CURLOPT_ACCEPT_ENCODING, "gzip");
+		// Rewind method, handle special error case where request need to rewind data stream
+		curl_easy_setopt(Curl, CURLOPT_SEEKFUNCTION, StaticSeekFn);
+		curl_easy_setopt(Curl, CURLOPT_SEEKDATA, this);
+		// Debug hooks
+#if UE_MIRAGE_DEBUG
+		curl_easy_setopt(Curl, CURLOPT_DEBUGDATA, this);
+		curl_easy_setopt(Curl, CURLOPT_DEBUGFUNCTION, StaticDebugCallback);
+		curl_easy_setopt(Curl, CURLOPT_VERBOSE, 1L);
+#endif // UE_MIRAGE_DEBUG
+	}
+
+	/** Gets the domain name for this request */
+	const FString& GetDomain() const
+	{
+		return Domain;
+	}
+
+	/** Returns the HTTP response code.*/
+	const int64 GetResponseCode() const
+	{
+		return ResponseCode;
+	}
+
+	/** Returns the number of bytes received this request (headers withstanding). */
+	const size_t GetBytesReceived() const
+	{
+		return BytesReceived;
+	}
+
+	/** Returns the number of bytes sent during this request (headers withstanding). */
+	const size_t  GetBytesSent() const
+	{
+		return BytesSent;
+	}
+
+	/**
+	* Upload buffer using the request, using either "Put" or "Post" verbs.
+	* @param Uri Url to use.
+	* @param Buffer Data to upload
+	* @return Result of the request
+	*/
+	template<RequestVerb V>
+	Result PerformBlockingUpload(const TCHAR* Uri, TArrayView<const uint8> Buffer)
+	{
+		static_assert(V == Put || V == PutJson || V == Post || V == PostJson, "Upload should use either Put or Post verbs.");
+
+		uint32 ContentLength = 0u;
+		EContentType ContentType = EContentType::None;
+
+		// TODO: Clean up the conditions here, too much code duplication
+		if (V == Put|| V == PutJson)
+		{
+			curl_easy_setopt(Curl, CURLOPT_UPLOAD, 1L);
+			curl_easy_setopt(Curl, CURLOPT_INFILESIZE, Buffer.Num());
+			curl_easy_setopt(Curl, CURLOPT_READDATA, this);
+			curl_easy_setopt(Curl, CURLOPT_READFUNCTION, StaticReadFn);
+
+			ContentType = V == Put ? EContentType::OctetStream : EContentType::Json;
+			ContentLength = Buffer.Num();
+			ReadDataView = Buffer;
+		}
+		else if (V == Post || V == PostJson)
+		{
+			curl_easy_setopt(Curl, CURLOPT_POST, 1L);
+			curl_easy_setopt(Curl, CURLOPT_INFILESIZE, Buffer.Num());
+			curl_easy_setopt(Curl, CURLOPT_READDATA, this);
+			curl_easy_setopt(Curl, CURLOPT_READFUNCTION, StaticReadFn);
+
+			ContentType = V == Post ? EContentType::UrlEncoded : EContentType::Json;
+			ContentLength = Buffer.Num();
+			ReadDataView = Buffer;
+		}
+
+		return PerformBlocking(Uri, V, ContentType, ContentLength);
+	}
+
+	/**
+	* Download an url into a buffer using the request.
+	* @param Uri Url to use.
+	* @param Buffer Optional buffer where data should be downloaded to. If empty downloaded data will
+	* be stored in an internal buffer and accessed GetResponse* methods.
+	* @return Result of the request
+	*/
+	Result PerformBlockingDownload(const TCHAR* Uri)
+	{
+		curl_easy_setopt(Curl, CURLOPT_HTTPGET, 1L);
+
+		WriteDataBufferPtr = nullptr;
+		WriteDataView.Reset();
+
+		return PerformBlocking(Uri, Get, EContentType::None, 0u);
+	}
+
+	Result PerformBlockingDownload(const TCHAR* Uri, TArray<uint8>* Buffer)
+	{
+		curl_easy_setopt(Curl, CURLOPT_HTTPGET, 1L);
+		
+		WriteDataBufferPtr = Buffer;
+		WriteDataView.Reset();
+
+		return PerformBlocking(Uri, Get, EContentType::None, 0u);
+	}
+
+	Result PerformBlockingDownload(const TCHAR* Uri, FMutableMemoryView Buffer)
+	{
+		curl_easy_setopt(Curl, CURLOPT_HTTPGET, 1L);
+
+		WriteDataBufferPtr = nullptr;
+		WriteDataView = Buffer;
+
+		return PerformBlocking(Uri, Get, EContentType::None, 0u);
+	}	
+
+	/**
+	* Query an url using the request. Queries can use either "Head" or "Delete" verbs.
+	* @param Uri Url to use.
+	* @return Result of the request
+	*/
+	template<RequestVerb V>
+	Result PerformBlockingQuery(const TCHAR* Uri)
+	{
+		static_assert(V == Head || V == Delete, "Queries should use either Head or Delete verbs.");
+
+		if (V == Delete)
+		{
+			curl_easy_setopt(Curl, CURLOPT_POST, 1L);
+			curl_easy_setopt(Curl, CURLOPT_CUSTOMREQUEST, "DELETE");
+		}
+		else if (V == Head)
+		{
+			curl_easy_setopt(Curl, CURLOPT_NOBODY, 1L);
+		}
+
+		return PerformBlocking(Uri, V, EContentType::None, 0u);
+	}
+
+	/**
+	* Set a header to send with the request.
+	*/
+	void SetHeader(const TCHAR* Header, const TCHAR* Value)
+	{
+		check(CurlResult == CURL_LAST); // Cannot set header after request is sent
+		Headers.Add(FString::Printf(TEXT("%s: %s"), Header, Value));
+	}
+
+	/**
+	* Attempts to find the header from the response. Returns false if header is not present.
+	*/
+	bool GetHeader(const ANSICHAR* Header, FString& OutValue) const
+	{
+		check(CurlResult != CURL_LAST);  // Cannot query headers before request is sent
+
+		const ANSICHAR* HeadersBuffer = (const ANSICHAR*)ResponseHeader.GetData();
+		size_t HeaderLen = strlen(Header);
+
+		// Find the header key in the (ANSI) response buffer. If not found we can exist immediately
+		if (const ANSICHAR* Found = strstr(HeadersBuffer, Header))
+		{
+			const ANSICHAR* Linebreak = strchr(Found, '\r');
+			const ANSICHAR* ValueStart = Found + HeaderLen + 2; //colon and space
+			const size_t ValueSize = Linebreak - ValueStart;
+			FUTF8ToTCHAR TCHARData(ValueStart, ValueSize);
+			OutValue = FString(TCHARData.Length(), TCHARData.Get());
+			return true;
+		}
+		return false;
+	}
+
+	/**
+	* Returns the response buffer. Note that is the request is performed
+	* with an external buffer as target buffer this string will be empty.
+	*/
+	const TArray<uint8>& GetResponseBuffer() const
+	{
+		return ResponseBuffer;
+	}
+
+	/**
+	* Returns the response buffer as a string. Note that is the request is performed
+	* with an external buffer as target buffer this string will be empty.
+	*/
+	FString GetResponseAsString() const
+	{
+		return GetAnsiBufferAsString(ResponseBuffer);
+	}
+
+	/**
+	* Returns the response header as a string.
+	*/
+	FString GetResponseHeaderAsString()
+	{
+		return GetAnsiBufferAsString(ResponseHeader);
+	}
+
+	/**
+	* Tries to parse the response buffer as a JsonObject. Return empty pointer if
+	* parse error occurs.
+	*/
+	TSharedPtr<FJsonObject> GetResponseAsJsonObject() const
+	{
+		FString Response = GetAnsiBufferAsString(ResponseBuffer);
+
+		TSharedPtr<FJsonObject> JsonObject;
+		TSharedRef<TJsonReader<> > JsonReader = TJsonReaderFactory<>::Create(Response);
+		if (!FJsonSerializer::Deserialize(JsonReader, JsonObject) || !JsonObject.IsValid())
+		{
+			return TSharedPtr<FJsonObject>(nullptr);
+		}
+
+		return JsonObject;
+	}
+
+	/** Will return true if the response code is considered a success */
+	static bool IsSuccessfulResponse(int64 ResponseCode)
+	{
+		// We consider anything in the 1XX or 2XX range a success
+		return ResponseCode >= 100 && ResponseCode < 300;
+	}
+
+private:
+
+	CURL* Curl;
+	CURLcode				CurlResult;
+	long					ResponseCode;
+	size_t					BytesSent;
+	size_t					BytesReceived;
+	bool					bLogErrors;
+
+	TArrayView<const uint8>	ReadDataView;
+
+	FMutableMemoryView		WriteDataView;
+	TArray<uint8>*			WriteDataBufferPtr;
+	TArray<uint8>*			WriteHeaderBufferPtr;
+
+	TArray<uint8>			ResponseHeader;
+	TArray<uint8>			ResponseBuffer;
+	TArray<FString>			Headers;
+	FString					Domain;
+	FAccessToken* AuthorizationToken;
+
+	/**
+	* Performs the request, blocking until finished.
+	* @param Uri Address on the domain to query
+	* @param Verb HTTP verb to use
+	* @param Buffer Optional buffer to directly receive the result of the request.
+	* If unset the response body will be stored in the request.
+	*/
+	Result PerformBlocking(const TCHAR* Uri, RequestVerb Verb, EContentType ContentType, uint32 ContentLength)
+	{
+		static const char* CommonHeaders[] = {
+			"User-Agent: UE4",
+			nullptr
+		};
+
+		TRACE_CPUPROFILER_EVENT_SCOPE(FRequest::PerformBlocking);
+
+		// Reset a few values incase the request has been reused (re-submitting a failure etc)
+		BytesSent = 0;
+		BytesReceived = 0;
+
+		// Setup request options
+		TStringBuilder<256> Url;
+		Url << Domain << TEXT("/") << Uri;
+
+		curl_easy_setopt(Curl, CURLOPT_URL, TCHAR_TO_ANSI(*Url));
+
+		// Setup response header buffer. If caller has not setup a response data buffer, use internal.
+		WriteHeaderBufferPtr = &ResponseHeader;
+		if (WriteDataBufferPtr == nullptr && WriteDataView.IsEmpty())
+		{
+			WriteDataBufferPtr = &ResponseBuffer;
+		}
+
+		// Build headers list
+		curl_slist* CurlHeaders = nullptr;
+		// Add common headers
+		for (uint8 i = 0; CommonHeaders[i] != nullptr; ++i)
+		{
+			CurlHeaders = curl_slist_append(CurlHeaders, CommonHeaders[i]);
+		}
+		
+		// Content-Length should always be set
+		TAnsiStringBuilder<48> ContentLengthString;
+		ContentLengthString.Appendf("Content-Length: %d", ContentLength);
+		CurlHeaders = curl_slist_append(CurlHeaders, *ContentLengthString);
+
+		// TODO: Clean up
+		switch (ContentType)
+		{
+		case OctetStream:
+			CurlHeaders = curl_slist_append(CurlHeaders, "Content-Type: application/octet-stream");
+			break;
+		case UrlEncoded:
+			CurlHeaders = curl_slist_append(CurlHeaders, "Content-Type: application/x-www-form-urlencoded");
+			break;
+		case Json:
+			CurlHeaders = curl_slist_append(CurlHeaders, "Content-Type: application/json");
+			break;
+		case Xml:
+			CurlHeaders = curl_slist_append(CurlHeaders, "Content-Type: application/xml");
+			break;
+		}
+
+		// Setup added headers
+		for (const FString& Header : Headers)
+		{
+			CurlHeaders = curl_slist_append(CurlHeaders, TCHAR_TO_ANSI(*Header));
+		}
+
+		// And auth token if it's set
+		if (AuthorizationToken != nullptr)
+		{
+			CurlHeaders = curl_slist_append(CurlHeaders, TCHAR_TO_ANSI(*AuthorizationToken->GetHeader()));
+		}
+
+		curl_easy_setopt(Curl, CURLOPT_HTTPHEADER, CurlHeaders);
+
+		// Shots fired!
+		CurlResult = curl_easy_perform(Curl);
+
+		// Get response code
+		bool bRedirected = false;
+		if (CURLE_OK == curl_easy_getinfo(Curl, CURLINFO_RESPONSE_CODE, &ResponseCode))
+		{
+			bRedirected = (ResponseCode >= 300 && ResponseCode < 400);
+		}
+
+		LogResult(CurlResult, Uri, Verb);
+
+		// Clean up
+		curl_slist_free_all(CurlHeaders);
+
+		return CurlResult == CURLE_OK ? Success : Failed;
+	}
+
+	void LogResult(CURLcode Result, const TCHAR* Uri, RequestVerb Verb) const
+	{
+		TRACE_CPUPROFILER_EVENT_SCOPE(FRequest::LogResult);
+
+		if (Result == CURLE_OK)
+		{
+			bool bSuccess = false;
+			const TCHAR* VerbStr = nullptr;
+			TStringBuilder<256> AdditionalInfo;
+
+			switch (Verb)
+			{
+			case Head:
+				// TODO: Io returns 404 if the head request is not found, Europa returns 400. Clean this up once 
+				// the inconsistency on the Jupiter server is fixed.
+				bSuccess = (ResponseCode == 400 || ResponseCode == 404 || IsSuccessfulResponse(ResponseCode));
+				VerbStr = TEXT("querying");
+				break;
+			case Get:
+				bSuccess = (ResponseCode == 400 || IsSuccessfulResponse(ResponseCode));
+				VerbStr = TEXT("fetching");
+				AdditionalInfo.Appendf(TEXT("Received: %d bytes."), BytesReceived);
+				break;
+			case Put:
+			case PutJson:
+				bSuccess = IsSuccessfulResponse(ResponseCode);
+				VerbStr = TEXT("updating");
+				AdditionalInfo.Appendf(TEXT("Sent: %d bytes."), BytesSent);
+				break;
+			case Post:
+			case PostJson:
+				bSuccess = IsSuccessfulResponse(ResponseCode);
+				VerbStr = TEXT("posting");
+				break;
+			case Delete:
+				bSuccess = IsSuccessfulResponse(ResponseCode);
+				VerbStr = TEXT("deleting");
+				break;
+			default:
+				checkf(false, TEXT("Unknown RequestVerb found in FRequet::LogResult"));
+			}
+
+			if (bSuccess)
+			{
+				UE_LOG(
+					LogVirtualization,
+					Verbose,
+					TEXT("Finished %s HTTP cache entry (response %d) from %s. %s"),
+					VerbStr,
+					ResponseCode,
+					Uri,
+					*AdditionalInfo
+				);
+			}
+			else if (bLogErrors)
+			{
+				// Print the response body if we got one, otherwise print header.
+				FString Response = GetAnsiBufferAsString(ResponseBuffer.Num() > 0 ? ResponseBuffer : ResponseHeader);
+				Response.ReplaceCharInline('\n', ' ');
+				Response.ReplaceCharInline('\r', ' ');
+				
+				// Don't log access denied as error, since tokens can expire mid session
+				if (ResponseCode == 401)
+				{
+					UE_LOG(
+						LogVirtualization,
+						Verbose,
+						TEXT("Failed %s HTTP cache entry (response %d) from %s. Response: %s"),
+						VerbStr,
+						ResponseCode,
+						Uri,
+						*Response
+					);
+				}
+				else
+				{	
+					UE_LOG(
+						LogVirtualization,
+						Error,
+						TEXT("Failed %s HTTP cache entry (response %d) from %s. Response: %s"),
+						VerbStr,
+						ResponseCode,
+						Uri,
+						*Response
+					);
+				}
+			}
+		}
+		else if (bLogErrors)
+		{
+			UE_LOG(
+				LogVirtualization,
+				Error,
+				TEXT("Error while connecting to %s: %s"),
+				*Domain,
+				ANSI_TO_TCHAR(curl_easy_strerror(Result))
+			);
+		}
+	}
+
+	FString GetAnsiBufferAsString(const TArray<uint8>& Buffer) const
+	{
+		// Content is NOT null-terminated; we need to specify lengths here
+		FUTF8ToTCHAR TCHARData(reinterpret_cast<const ANSICHAR*>(Buffer.GetData()), Buffer.Num());
+		return FString(TCHARData.Length(), TCHARData.Get());
+	}
+
+#if UE_MIRAGE_DEBUG
+	static size_t StaticDebugCallback(CURL* Handle, curl_infotype DebugInfoType, char* DebugInfo, size_t DebugInfoSize, void* UserData)
+	{
+		FRequest* Request = static_cast<FRequest*>(UserData);
+
+		auto PrintText = [](FRequest* InRequest, char* InDebugInfo, size_t InDebugInfoSize)
+		{
+			// Truncate at 1023 characters. This is just an arbitrary number based on a buffer size seen in
+			// the libcurl code.
+			InDebugInfoSize = FMath::Min(InDebugInfoSize, (size_t)1023);
+
+			// Calculate the actual length of the string due to incorrect use of snprintf() in lib/vtls/openssl.c.
+			char* FoundNulPtr = (char*)memchr(InDebugInfo, 0, InDebugInfoSize);
+			int CalculatedSize = FoundNulPtr != nullptr ? FoundNulPtr - InDebugInfo : InDebugInfoSize;
+
+			auto ConvertedString = StringCast<TCHAR>(static_cast<const ANSICHAR*>(InDebugInfo), CalculatedSize);
+			FString DebugText(ConvertedString.Length(), ConvertedString.Get());
+
+			UE_LOG(LogVirtualization, Log, TEXT("%p: '%s'"), InRequest, *DebugText);
+		};
+
+		switch (DebugInfoType)
+		{
+		case CURLINFO_TEXT:
+		{
+			PrintText(Request, DebugInfo, DebugInfoSize);
+		}
+		break;
+
+		case CURLINFO_HEADER_IN:
+			UE_LOG(LogVirtualization, Log, TEXT("%p: Received header (%d bytes)"), Request, DebugInfoSize);
+			break;
+
+		case CURLINFO_HEADER_OUT:
+			UE_LOG(LogVirtualization, Log, TEXT("%p: Sent header (%d bytes)"), Request, DebugInfoSize);
+			PrintText(Request, DebugInfo, DebugInfoSize);
+			break;
+
+		case CURLINFO_DATA_IN:
+			UE_LOG(LogVirtualization, Log, TEXT("%p: Received data (%d bytes)"), Request, DebugInfoSize);
+			break;
+
+		case CURLINFO_DATA_OUT:
+			UE_LOG(LogVirtualization, Log, TEXT("%p: Sent data (%d bytes)"), Request, DebugInfoSize);
+			break;
+
+		case CURLINFO_SSL_DATA_IN:
+			UE_LOG(LogVirtualization, Log, TEXT("%p: Received SSL data (%d bytes)"), Request, DebugInfoSize);
+			break;
+
+		case CURLINFO_SSL_DATA_OUT:
+			UE_LOG(LogVirtualization, Log, TEXT("%p: Sent SSL data (%d bytes)"), Request, DebugInfoSize);
+			break;
+		}
+
+		return 0;
+	}
+#endif // UE_MIRAGE_DEBUG
+
+	static size_t StaticReadFn(void* Ptr, size_t SizeInBlocks, size_t BlockSizeInBytes, void* UserData)
+	{
+		TRACE_CPUPROFILER_EVENT_SCOPE(FRequest::StaticReadFn);
+		
+		FRequest* Request = static_cast<FRequest*>(UserData);
+		TArrayView<const uint8>& ReadDataView = Request->ReadDataView;
+
+		const size_t Offset = Request->BytesSent;
+		const size_t ReadSize = FMath::Min((size_t)ReadDataView.Num() - Offset, SizeInBlocks * BlockSizeInBytes);
+		check(ReadDataView.Num() >= Offset + ReadSize);
+
+		FMemory::Memcpy(Ptr, ReadDataView.GetData() + Offset, ReadSize);
+		Request->BytesSent += ReadSize;
+
+		return ReadSize;
+	}
+
+	static size_t StaticWriteHeaderFn(void* Ptr, size_t SizeInBlocks, size_t BlockSizeInBytes, void* UserData)
+	{
+		TRACE_CPUPROFILER_EVENT_SCOPE(FRequest::StaticWriteHeaderFn);
+
+		FRequest* Request = static_cast<FRequest*>(UserData);
+		const size_t WriteSize = SizeInBlocks * BlockSizeInBytes;
+		TArray<uint8>* WriteHeaderBufferPtr = Request->WriteHeaderBufferPtr;
+		if (WriteHeaderBufferPtr && WriteSize > 0)
+		{
+			const size_t CurrentBufferLength = WriteHeaderBufferPtr->Num();
+			if (CurrentBufferLength > 0)
+			{
+				// Remove the previous zero termination
+				(*WriteHeaderBufferPtr)[CurrentBufferLength - 1] = ' ';
+			}
+
+			// Write the header
+			WriteHeaderBufferPtr->Append((const uint8*)Ptr, WriteSize + 1);
+			(*WriteHeaderBufferPtr)[WriteHeaderBufferPtr->Num() - 1] = 0; // Zero terminate string
+			return WriteSize;
+		}
+		return 0;
+	}
+
+	static size_t StaticWriteBodyFn(void* Ptr, size_t SizeInBlocks, size_t BlockSizeInBytes, void* UserData)
+	{
+		TRACE_CPUPROFILER_EVENT_SCOPE(FRequest::StaticWriteBodyFn);
+
+		FRequest* Request = static_cast<FRequest*>(UserData);
+		const size_t WriteSize = SizeInBlocks * BlockSizeInBytes;
+
+		if (WriteSize == 0)
+		{
+			return 0;
+		}
+
+		if (TArray<uint8>* WriteDataBufferPtr = Request->WriteDataBufferPtr)
+		{
+			// If this is the first part of the body being received, try to reserve 
+			// memory if content length is defined in the header.
+			if (Request->BytesReceived == 0 && Request->WriteHeaderBufferPtr)
+			{
+				static const ANSICHAR* ContentLengthHeaderStr = "Content-Length: ";
+				const ANSICHAR* Header = (const ANSICHAR*)Request->WriteHeaderBufferPtr->GetData();
+
+				if (const ANSICHAR* ContentLengthHeader = FCStringAnsi::Strstr(Header, ContentLengthHeaderStr))
+				{
+					size_t ContentLength = (size_t)FCStringAnsi::Atoi64(ContentLengthHeader + strlen(ContentLengthHeaderStr));
+					if (ContentLength > 0u && ContentLength < UE_MIRAGE_MAX_BUFFER_RESERVE)
+					{
+						WriteDataBufferPtr->Reserve(ContentLength);
+					}
+				}
+			}
+
+			// Write to the target buffer
+			WriteDataBufferPtr->Append((const uint8*)Ptr, WriteSize);
+			Request->BytesReceived += WriteSize;
+			return WriteSize;
+		}
+		else if (!Request->WriteDataView.IsEmpty())
+		{
+			if (Request->BytesReceived + WriteSize <= Request->WriteDataView.GetSize())
+			{
+				FMemory::Memcpy((uint8*)Request->WriteDataView.GetData() + Request->BytesReceived, Ptr, WriteSize);
+				Request->BytesReceived += WriteSize;
+
+				return WriteSize;
+			}
+			else
+			{
+				UE_CLOG(
+					Request->bLogErrors,
+					LogVirtualization,
+					Error,
+					TEXT("Attempting to write %d bytes to the response buffer which only has %" SSIZE_T_FMT " bytes remaining %" SSIZE_T_FMT),
+					WriteSize,
+					(size_t)Request->WriteDataView.GetSize() - Request->BytesReceived
+				);
+
+				return -1;
+			}
+		}
+		else
+		{
+			UE_CLOG(Request->bLogErrors, LogVirtualization, Error, TEXT("No response buffer was set!"));	
+			return -1;
+		}
+	}
+
+	static size_t StaticSeekFn(void* UserData, curl_off_t Offset, int Origin)
+	{
+		TRACE_CPUPROFILER_EVENT_SCOPE(FRequest::StaticSeekFn);
+
+		FRequest* Request = static_cast<FRequest*>(UserData);
+		size_t NewPosition = 0;
+
+		switch (Origin)
+		{
+		case SEEK_SET: NewPosition = Offset; break;
+		case SEEK_CUR: NewPosition = Request->BytesSent + Offset; break;
+		case SEEK_END: NewPosition = Request->ReadDataView.Num() + Offset; break;
+		}
+
+		// Make sure we don't seek outside of the buffer
+		if (NewPosition < 0 || NewPosition >= Request->ReadDataView.Num())
+		{
+			return CURL_SEEKFUNC_FAIL;
+		}
+
+		// Update the used offset
+		Request->BytesSent = NewPosition;
+		return CURL_SEEKFUNC_OK;
+	}
+};
+
+/**
+* Pool that manages a fixed set of requests. Users are required to release requests that have been 
+* acquired. Usable with \ref FScopedRequestPtr which handles this automatically.
+*/
+struct FRequestPool
+{
+	FRequestPool(const TCHAR* InServiceUrl, FAccessToken* InAuthorizationToken)
+	{
+		for (uint8 i = 0; i < Pool.Num(); ++i)
+		{
+			Pool[i].Usage = 0u;
+			Pool[i].Request = new FRequest(InServiceUrl, InAuthorizationToken, true);	
+		}
+	}
+
+	~FRequestPool()
+	{
+		for (uint8 i = 0; i < Pool.Num(); ++i)
+		{
+			// No requests should be in use by now.
+			check(Pool[i].Usage.Load(EMemoryOrder::Relaxed) == 0u);
+			delete Pool[i].Request;
+		}
+	}
+
+	/**
+	* Block until a request is free. Once a request has been returned it is 
+	* "owned by the caller and need to release it to the pool when work has been completed.
+	* @return Usable request instance.
+	*/
+	FRequest* WaitForFreeRequest()
+	{
+		TRACE_CPUPROFILER_EVENT_SCOPE(FRequestPool::WaitForFreeRequest);
+		while (true)
+		{
+			for (uint8 i = 0; i < Pool.Num(); ++i)
+			{
+				if (!Pool[i].Usage.Load(EMemoryOrder::Relaxed))
+				{
+					uint8 Expected = 0u;
+					if (Pool[i].Usage.CompareExchange(Expected, 1u))
+					{
+						return Pool[i].Request;
+					}
+				}
+			}
+
+			TRACE_COUNTER_ADD(Jupiter_WaitOnRequestPool, int64(1));
+			FPlatformProcess::Sleep(UE_MIRAGE_REQUEST_POOL_WAIT_INTERVAL);
+		}
+	}
+
+	/**
+	* Re lease request to the pool.
+	* @param Request Request that should be freed. Note that any buffer owned by the request can now be reset.
+	*/
+	void ReleaseRequestToPool(FRequest* Request)
+	{
+		for (uint8 i = 0; i < Pool.Num(); ++i)
+		{
+			if (Pool[i].Request == Request)
+			{
+				Request->Reset();
+				Pool[i].Usage.Exchange(0u);
+				return;
+			}
+		}
+		check(false);
+	}
+
+private:
+
+	struct FEntry
+	{
+		TAtomic<uint8> Usage;
+		FRequest* Request;
+	};
+
+	TStaticArray<FEntry, UE_MIRAGE_REQUEST_POOL_SIZE> Pool;
+
+	FRequestPool() {}
+};
+
+/**
+* Utility class to manage requesting and releasing requests from the \ref FRequestPool.
+*/
+struct FScopedRequestPtr
+{
+public:
+	FScopedRequestPtr(FRequestPool* InPool)
+		: Request(InPool->WaitForFreeRequest())
+		, Pool(InPool)
+	{}
+
+	~FScopedRequestPtr()
+	{
+		Reset();
+	}
+
+	bool IsValid() const 
+	{
+		return Request != nullptr;
+	}
+
+	FRequest* Get() const
+	{
+		check(IsValid());
+		return Request;
+	}
+
+	FRequest* operator->()
+	{
+		check(IsValid());
+		return Request;
+	}
+
+	void Reset()
+	{
+		if (IsValid())
+		{
+			Pool->ReleaseRequestToPool(Request);
+			Request = nullptr;
+		}
+	}
+
+private:
+	FRequest* Request;
+	FRequestPool* Pool;
+};
+
+
+/**
+* Adds a checksum (as request header) for a given payload. Jupiter will use this to verify the integrity
+* of the received data.
+* @param Request Request that the data will be sent with.
+* @param Payload Payload that will be sent.
+* @return True on success, false on failure.
+*/
+FSHAHash HashPayload(FRequest* Request, const TArrayView<const uint8> Payload)
+{
+	FSHAHash PayloadHash;
+	FSHA1::HashBuffer(Payload.GetData(), Payload.Num(), PayloadHash.Hash);
+	Request->SetHeader(TEXT("X-Jupiter-Sha1"), *PayloadHash.ToString());
+	return PayloadHash;
+}
+
+/**
+* Verifies the integrity of the received data using supplied checksum.
+* @param Hash Received hash value.
+* @param Payload Payload received.
+* @return True if the data is correct, false if checksums doesn't match.
+*/
+bool VerifyPayload(FSHAHash Hash, const TArray<uint8>& Payload)
+{
+	FSHAHash PayloadHash;
+	FSHA1::HashBuffer(Payload.GetData(), Payload.Num(), PayloadHash.Hash);
+
+	if (Hash != PayloadHash)
+	{
+		UE_LOG(LogVirtualization,
+			Warning,
+			TEXT("Checksum from server did not match recieved data (%s vs %s). Discarding cached result."),
+			*Hash.ToString(),
+			*PayloadHash.ToString()
+		);
+		return false;
+	}
+
+	return true;
+}
+
+/**
+* Verifies the integrity of the received data using supplied checksum.
+* @param Request Request that the data was be received with.
+* @param Payload Payload received.
+* @return True if the data is correct, false if checksums doesn't match.
+*/
+bool VerifyRequest(const FRequest* Request, const TArray<uint8>& Payload)
+{
+	FString RecievedHashStr;
+	FSHAHash PayloadHash;
+	FSHA1::HashBuffer(Payload.GetData(), Payload.Num(), PayloadHash.Hash);
+	if (Request->GetHeader("X-Jupiter-Sha1", RecievedHashStr))
+	{
+		FSHAHash RecievedHash;
+		RecievedHash.FromString(RecievedHashStr);
+		return VerifyPayload(RecievedHash, Payload);
+	}
+	UE_LOG(LogVirtualization, Error, TEXT("HTTP server did not send a content hash. Wrong server version?"));
+	return true;
+}
+
+#if WITH_SSL
+static int SslCertVerify(int PreverifyOk, X509_STORE_CTX* Context)
+{
+	if (PreverifyOk == 1)
+	{
+		SSL* Handle = static_cast<SSL*>(X509_STORE_CTX_get_ex_data(Context, SSL_get_ex_data_X509_STORE_CTX_idx()));
+		check(Handle);
+
+		SSL_CTX* SslContext = SSL_get_SSL_CTX(Handle);
+		check(SslContext);
+
+		FRequest* Request = static_cast<FRequest*>(SSL_CTX_get_app_data(SslContext));
+		check(Request);
+
+		const FString& Domain = Request->GetDomain();
+
+		if (!FSslModule::Get().GetCertificateManager().VerifySslCertificates(Context, Domain))
+		{
+			PreverifyOk = 0;
+		}
+	}
+
+	return PreverifyOk;
+}
+
+static CURLcode sslctx_function(CURL* curl, void* sslctx, void* parm)
+{
+	SSL_CTX* Context = static_cast<SSL_CTX*>(sslctx);
+	const ISslCertificateManager& CertificateManager = FSslModule::Get().GetCertificateManager();
+
+	CertificateManager.AddCertificatesToSslContext(Context);
+	SSL_CTX_set_verify(Context, SSL_CTX_get_verify_mode(Context), SslCertVerify);
+	SSL_CTX_set_app_data(Context, parm);
+
+	/* all set to go */
+	return CURLE_OK;
+}
+#endif //WITH_SSL
+
+} // namespace Utility
+
+namespace Virtualization
+{
+
+/** Represents the data required to make a PUT request to the Europa DDCCache API
+	that can be easily serialized to JSON to make the actual request. */
+struct FEuropaDDCCachePUTRequest : public FJsonSerializable
+{
+	struct FMetaData : public FJsonSerializable
+	{
+		BEGIN_JSON_SERIALIZER
+			JSON_SERIALIZE("payloadLength", PayloadLength);
+			JSON_SERIALIZE("chunkLength", ChunkLength);
+		END_JSON_SERIALIZER
+
+		/** Overall length (in bytes) of the payload */
+		int64 PayloadLength { INDEX_NONE };
+		/** The max length (in bytes) of each chunk */
+		int64 ChunkLength { INDEX_NONE };
+	};
+
+	BEGIN_JSON_SERIALIZER
+		JSON_SERIALIZE_ARRAY("blobReferences", ChunkHashes);
+		JSON_SERIALIZE_OBJECT_SERIALIZABLE("metadata", MetaData);
+		JSON_SERIALIZE("contentHash", PayloadHash);
+	END_JSON_SERIALIZER
+
+	/** A FSHAHash of the entire payload once it is reconstructed */
+	FString PayloadHash;
+	/** An object containing additional metadata */
+	FMetaData MetaData;
+	/** Each string represents a FSHAHash */
+	TArray<FString> ChunkHashes;
+};
+
+/** Data structure for use with Jupiter Europa DDCCache GET JSON response */
+struct FDDCCacheGETResponse : public FJsonSerializable
+{
+	struct FMetaData : public FJsonSerializable
+	{
+		BEGIN_JSON_SERIALIZER
+			JSON_SERIALIZE("payloadLength", PayloadLength);
+			JSON_SERIALIZE("chunkLength", ChunkLength);
+		END_JSON_SERIALIZER
+
+		/** Overall length of the payload */
+		int64 PayloadLength { INDEX_NONE };
+		/** The max length of each chunk */
+		int64 ChunkLength { INDEX_NONE };
+	};
+
+	BEGIN_JSON_SERIALIZER
+		JSON_SERIALIZE("name", Name);
+		JSON_SERIALIZE("lastAccessTime", LastAccessTime);
+		JSON_SERIALIZE_OBJECT_SERIALIZABLE("metadata", MetaData);
+		JSON_SERIALIZE("contentHash", PayloadHash);
+		JSON_SERIALIZE_ARRAY("blobIdentifiers", ChunkHashes);
+		JSON_SERIALIZE("blob", PayloadBlob);
+	END_JSON_SERIALIZER
+
+	/** The payload name in the format: '{namespace}.{bucket}.{key}' */
+	FString Name;
+	/** The data and time that the payload was last accessed */
+	FDateTime LastAccessTime { INDEX_NONE };
+	/** A FSHAHash of the entire payload once it is reconstructed */
+	FString PayloadHash;
+	/** An object containing additional metadata */
+	FMetaData MetaData;
+	/** Each string represents a FSHAHash */
+	TArray<FString> ChunkHashes;
+	/** The payload (Base64 encoding) */
+	FString PayloadBlob;
+};
+
+/** Data structure for use with Jupiter statuc GET response. Represents the status of the Jupiter service. */
+struct FJupiterServiceStatus : public FJsonSerializable
+{
+	/** Returns true if the current version in the object is greater or equal to the version numbers passed in */
+	bool DoesHaveValidVersion(uint32 MinMajorVersion, uint32 MinMinorVersion, uint32 MinPatchVersion) const 
+	{
+		uint32 MajorVersion = 0;
+		uint32 MinorVersion = 0;
+		uint32 PatchVersion = 0;
+
+		if (GetVersionNumbers(MajorVersion, MinorVersion, PatchVersion))
+		{
+			if (IsValidVersion(MinMajorVersion, MinMinorVersion, MinPatchVersion, MajorVersion, MinorVersion, PatchVersion))
+			{
+				return true;
+			}
+			else
+			{
+				UE_LOG(LogVirtualization, Error, TEXT("Jupiter service version is too old! Found: '%s' Required: %u.%u.%u"), 
+					*Version, MajorVersion, MinorVersion, PatchVersion);
+				return false;
+			}
+		}
+		else
+		{
+			UE_LOG(LogVirtualization, Error, TEXT("Failed to parse valid Jupiter version number from '%s'"), *Version);
+			return false;
+		}
+	}
+
+	/** Returns the version numbers parsed from string format into uint32 values, returns false if the string cannot be parsed */
+	bool GetVersionNumbers(uint32& OutMajorVersion, uint32& OutMinorVersion, uint32& OutPatchVersion) const
+	{
+		TCHAR* End = nullptr;
+		// Read the major/minor/patch numbers
+		uint64 MajorVersion = FCString::Strtoui64(*Version, &End, 10);
+		if (MajorVersion > MAX_uint32 || *(End++) != '.')
+		{
+			return false;
+		}
+
+		uint64 MinorVersion = FCString::Strtoui64(End, &End, 10);
+		if (MinorVersion > MAX_uint32 || *(End++) != '.')
+		{
+			return false;
+		}
+
+		uint64 PatchVersion = FCString::Strtoui64(End, &End, 10);
+		if (PatchVersion > MAX_uint32)
+		{
+			return false;
+		}
+
+		OutMajorVersion = static_cast<uint16>(MajorVersion);
+		OutMinorVersion = static_cast<uint16>(MinorVersion);
+		OutPatchVersion = static_cast<uint16>(PatchVersion);
+
+		return true;
+	}
+
+	/** Returns true if the given capability is supported by the Jupiter service */
+	bool SupportsCapability(FStringView Capability) const 
+	{
+		return Capabilities.Contains(Capability);
+	}
+
+	/** Prints the status of the Jupiter service to the log */
+	void LogStatusInfo() const
+	{
+		UE_LOG(LogVirtualization, Log, TEXT("Jupiter Service Status:"));
+		UE_LOG(LogVirtualization, Log, TEXT("Version: %s"), *Version);
+		UE_LOG(LogVirtualization, Log, TEXT("GitHash: %s"), *GitHash);		
+		UE_LOG(LogVirtualization, Log, TEXT("Capabilities:"));
+		for (const FString& Capability : Capabilities)
+		{
+			UE_LOG(LogVirtualization, Log, TEXT("\t%s"), *Capability);
+		}
+		UE_LOG(LogVirtualization, Log, TEXT("Site Id: %s"), *SiteIdentifier);
+	}
+
+private:
+
+	bool IsValidVersion(uint32 MinMajorVersion, uint32 MinMinorVersion, uint32 MinPatchVersion,
+						uint32 CurrentMajorVersion, uint32 CurrentMinorVersion, uint32 CurrentPatchVersion) const
+	{
+		if (CurrentMajorVersion != MinMajorVersion)
+		{
+			return CurrentMajorVersion > MinMajorVersion;
+		}
+
+		if (CurrentMinorVersion != MinMinorVersion)
+		{
+			return CurrentMinorVersion > MinMinorVersion;
+		}
+
+		if (CurrentPatchVersion != MinPatchVersion)
+		{
+			return CurrentPatchVersion > MinPatchVersion;
+		}
+
+		return true;
+	}
+
+	BEGIN_JSON_SERIALIZER
+		JSON_SERIALIZE("version", Version);
+		JSON_SERIALIZE("gitHash", GitHash);
+		JSON_SERIALIZE_ARRAY("capabilities", Capabilities);
+		JSON_SERIALIZE("siteIdentifier", SiteIdentifier);
+	END_JSON_SERIALIZER
+
+	/** Version of the service in the format [MAJOR.MINOR.PATCH] */
+	FString Version;
+	/**The git commit hash for the Jupiter service */
+	FString GitHash;
+	/** An array of which sub-services the Jupiter service supports */
+	TArray<FString> Capabilities;
+	/** The identifier for the server connected to */
+	FString SiteIdentifier;
+};
+
+/**
+* This backend allows data to be stored in and retrieved from a Jupiter service.
+* 
+* Ini file setup:
+* 'Name'=(Type=Jupiter, Host="", Namespace="", ChunkSize=, OAuthProvider="", OAuthClientId="", OAuthSecret="")
+* Host:				The URL of the service, use http://localhost if hosted locally.
+* Namespace:		Jupiter storage is divided into a number of namespaces allowing projects to keep their data separate
+*					while using the same service. This value controls which name space will be used.
+* ChunkSize:		Each payload can be divided into a number of chunks when being uploaded to Jupiter to improve upload 
+*					performance, this value sets the max size (in bytes) of each chunk. To disable and attempt to upload 
+*					each payload as a single data blob, set this to -1. 
+* OAuthProvider:	Url of the OAuth authorization server.
+* OAuthClientId:	Public identifier for use with the OAuth authorization server.
+* OAuthSecret:		Password for the OAuthClientId
+* (Note that the OAuth entries are not required if hosting locally)
+*/
+class FJupiterBackend : public IVirtualizationBackend
+{
+public:
+	FJupiterBackend(FStringView ConfigName)
+		: IVirtualizationBackend(EOperations::Both)
+		, Namespace(TEXT("mirage"))
+		, Bucket(TEXT("default"))
+		, ChunkSize(-1)
+		, FailedLoginAttempts(0)		
+	{
+		
+	}
+
+private:
+
+	virtual bool Initialize(const FString& ConfigEntry)
+	{
+		using namespace Utility;
+
+		// Some fields are required and will give fatal errors if not found!
+
+		if (!FParse::Value(*ConfigEntry, TEXT("Host="), HostAddress))
+		{
+			UE_LOG(LogVirtualization, Fatal, TEXT("'Host=' not found in the config file"));
+			return false;
+		}
+
+		if (!FParse::Value(*ConfigEntry, TEXT("Namespace="), Namespace))
+		{
+			UE_LOG(LogVirtualization, Fatal, TEXT("'Namespace=' not found in the config file"));
+			return false;
+		}
+
+		if (FParse::Value(*ConfigEntry, TEXT("ChunkSize="), ChunkSize))
+		{
+			UE_LOG(LogVirtualization, Log, TEXT("ChunkSize set to '%" UINT64_FMT "' bytes"), ChunkSize);
+		}
+		else
+		{
+			UE_LOG(LogVirtualization, Log, TEXT("Payloads will not be chunked!"));
+		}
+
+		// If we are connecting to a locally hosted Jupiter service then we do not need authorization
+		if (!IsUsingLocalHost())
+		{
+			if (!FParse::Value(*ConfigEntry, TEXT("OAuthProvider="), OAuthProvider))
+			{
+				UE_LOG(LogVirtualization, Fatal, TEXT("'OAuthProvider=' not found in the config file"));
+					return false;
+			}
+
+			if (!FParse::Value(*ConfigEntry, TEXT("OAuthSecret="), OAuthSecret))
+			{
+				UE_LOG(LogVirtualization, Fatal, TEXT("'OAuthSecret=' not found in the config file"));
+				return false;
+			}
+
+			if (!FParse::Value(*ConfigEntry, TEXT("OAuthClientId="), OAuthClientId))
+			{
+				UE_LOG(LogVirtualization, Fatal, TEXT("'OAuthClientId=' not found in the config file"));
+				return false;
+			}
+		}
+
+		UE_LOG(LogVirtualization, Log, TEXT("Attempting to connect to a Jupiter service at '%s' with namespace '%s'"), *HostAddress, *Namespace);
+
+		if (!IsServiceReady())
+		{
+			return false;
+		}
+
+		if (!AcquireAccessToken())
+		{
+			return false;
+		}
+
+		if (!ValidateServiceVersion())
+		{
+			return false;
+		}
+		
+		RequestPool = MakeUnique<FRequestPool>(*HostAddress, AccessToken.Get());
+
+		return true; 
+	}
+
+	virtual bool PushData(const FPayloadId& Id, const FCompressedBuffer& CompressedPayload) override
+	{
+		using namespace Utility;
+
+		TRACE_CPUPROFILER_EVENT_SCOPE(FJupiterBackend::PushData);
+
+#if UE_CHECK_FOR_EXISTING_PAYLOADS
+		if (DoesPayloadExist(Id))
+		{
+			UE_LOG(LogVirtualization, Verbose, TEXT("Jupiter already has a copy of the payload '%s'"), *Id.ToString());
+			return true;
+		}
+#endif // UE_CHECK_FOR_EXISTING_PAYLOADS
+
+		FEuropaDDCCachePUTRequest PUTRequest;
+
+		// TODO: This is a waste, we shouldn't need to flatten the buffer, but it makes it easier to work
+		// with the existing chunked code. The chunking code is most likely going to be removed before 
+		// this backend goes to production so it is not worth fixing up to work with FCompressedBuffer 
+		// properly.
+		FSharedBuffer FlattenedPayload = CompressedPayload.GetCompressed().Flatten();
+
+		const int64 NumChunks = FMath::DivideAndRoundUp(FlattenedPayload.GetSize(), ChunkSize);
+		if (NumChunks > MAX_int32)
+		{
+			UE_LOG(LogVirtualization, Error, TEXT("Too many chunks (%d) are required for the payload '%s', try increasing the ChunkSize"), NumChunks, *Id.ToString());
+			return false;
+		}
+
+		PUTRequest.ChunkHashes.SetNum((int32)NumChunks);
+
+		const uint8* DataPtr = (const uint8*)FlattenedPayload.GetData();
+
+		std::atomic<int32> NumFailedChunks(0);
+		FGraphEventArray Tasks;
+		Tasks.Reserve((int32)NumChunks);
+
+		// Create and process the chunks that make up the payload
+		for (int32 Index = 0; Index < NumChunks; ++Index )
+		{
+			const int64 ChunkStart = Index * ChunkSize;
+			const uint64 BytesInChunk = FMath::Min(ChunkSize, FlattenedPayload.GetSize() - ChunkStart);
+
+			TArrayView<const uint8> ChunkData(&DataPtr[ChunkStart], (int32)BytesInChunk);
+			FString& ChunkHashString = PUTRequest.ChunkHashes[Index];
+
+			auto Job = [this, ChunkData, &Id, &ChunkHashString, &NumFailedChunks]()
+			{
+				if (!PostChunk(ChunkData, Id, ChunkHashString))
+				{
+					NumFailedChunks++;
+				}
+			};
+
+#if UE_ENABLE_ASYNC_CHUNK_ACCESS
+			Tasks.Add(FFunctionGraphTask::CreateAndDispatchWhenReady(MoveTemp(Job)));
+#else
+			Job();
+#endif //UE_ENABLE_ASYNC_CHUNK_ACCESS
+		}
+
+		// There is some expensive work that we can do while we wait for the chunks to finish their upload
+		const FIoHash PayloadHash = FIoHash::HashBuffer(FlattenedPayload.GetData(), FlattenedPayload.GetSize());
+
+		PUTRequest.PayloadHash = LexToString(PayloadHash);
+		PUTRequest.MetaData.PayloadLength = (int64)FlattenedPayload.GetSize();
+		PUTRequest.MetaData.ChunkLength = ChunkSize;
+
+		{
+			TRACE_CPUPROFILER_EVENT_SCOPE(FJupiterBackend::PullData::WaitOnChunks);
+			FTaskGraphInterface::Get().WaitUntilTasksComplete(Tasks);
+		}
+
+		if (NumFailedChunks > 0)
+		{
+			UE_LOG(LogVirtualization, Error, TEXT("Failed to upload %d chunks for the payload '%s'."), NumFailedChunks.load(), *Id.ToString());
+			return false;
+		}
+
+		UE_LOG(LogVirtualization, Verbose, TEXT("Successfully uploaded all chunks for the payload '%s'"), *Id.ToString());	
+
+		// Note that the ddc end point is used by both ddc and mirage
+		TStringBuilder<256> Uri;
+		Uri.Appendf(TEXT("api/v1/c/ddc/%s/%s/%s"), *Namespace, *Bucket, *Id.ToString());
+
+		// Retry request until we get an accepted response or exhaust allowed number of attempts.
+		uint32 Attempts = 0;
+		while (++Attempts <= UE_MIRAGE_MAX_ATTEMPTS)
+		{
+			FScopedRequestPtr Request(RequestPool.Get());
+			if (Request.IsValid())
+			{	
+				const FString OutputString = PUTRequest.ToJson(false);
+				
+				TStringConversion<TStringConvert<TCHAR, ANSICHAR>> ConvertedString = StringCast<ANSICHAR>(*OutputString);
+				TArrayView<const uint8> AnsiStringBuffer((uint8*)ConvertedString.Get(),ConvertedString.Length());
+
+				FRequest::Result Result = Request->PerformBlockingUpload<FRequest::PutJson>(*Uri, AnsiStringBuffer);
+				const int64 ResponseCode = Request->GetResponseCode();
+
+				if (ResponseCode == 200)
+				{
+					UE_LOG(LogVirtualization, Verbose, TEXT("Successfully uploaded the description for the payload '%s'"), *Id.ToString());	
+					return true;
+				}
+
+				if (!ShouldRetryOnError(ResponseCode))
+				{
+					UE_LOG(LogVirtualization, Error, TEXT("Failed with error code '%d' to upload header infomation about payload '%s'"), ResponseCode, *Id.ToString());
+					return false;
+				}
+			}
+		}
+
+		UE_LOG(LogVirtualization, Error, TEXT("Failed  '%d' attempts to upload header infomation about payload '%s'"), UE_MIRAGE_MAX_ATTEMPTS, *Id.ToString());
+		return false;
+	}
+
+	virtual FCompressedBuffer PullData(const FPayloadId& Id) override
+	{
+		using namespace Utility;
+
+		TRACE_CPUPROFILER_EVENT_SCOPE(FJupiterBackend::PullData);
+
+		// First we need to get the description of the payload from Europa
+
+		// Note that the ddc end point is used by both ddc and mirage
+		// fields=contentSha1 - Ask for the Sha1 hash of the fully reconstructed payload
+		// fields=blobIdentifiers - As for a list of the Sha1 hash ids for the payload chunks that we need to access from
+		// the Io service, the ids will be in the correct order.
+		// fields=metadata - Ask for the payload metadata which contains info we can use later for optimizations.
+		TStringBuilder<256> Uri;
+		Uri.Appendf(TEXT("api/v1/c/ddc/%s/%s/%s.json?fields=contentHash&fields=blobIdentifiers&fields=metadata"), *Namespace, *Bucket, *Id.ToString());
+
+		FDDCCacheGETResponse Response;
+
+		// Retry request until we get an accepted response or exhaust allowed number of attempts.
+		int64 ResponseCode = 0;
+		uint32 Attempts = 0;
+		while (ResponseCode != 200 && ++Attempts <= UE_MIRAGE_MAX_ATTEMPTS)
+		{
+			FScopedRequestPtr Request(RequestPool.Get());
+			if (Request.IsValid())
+			{
+				const FRequest::Result Result = Request->PerformBlockingDownload(*Uri);
+				ResponseCode = Request->GetResponseCode();
+
+				// Request was successful, make sure we got all the expected data.
+				if (ResponseCode == 200)
+				{
+					if (!Response.FromJson(Request->GetResponseAsJsonObject()))
+					{
+						UE_LOG(LogVirtualization, Error, TEXT("No valid response to http GET found, when accessing payload '%s'"), *Id.ToString());
+						return FCompressedBuffer();
+					}
+				}
+				else if(!ShouldRetryOnError(ResponseCode))
+				{
+					UE_LOG(LogVirtualization, Error, TEXT("Failed with error code '%d' to download header infomation about payload '%s'"), ResponseCode, *Id.ToString());
+					return FCompressedBuffer();
+				}
+			}
+		}
+
+		if (ResponseCode != 200)
+		{
+			UE_LOG(LogVirtualization, Error, TEXT("Failed '%d' attempts to download header infomation about payload (last error code '%d')  '%s'"), UE_MIRAGE_MAX_ATTEMPTS, ResponseCode, *Id.ToString());
+			return FCompressedBuffer();
+		}
+
+		UE_LOG(LogVirtualization, Verbose, TEXT("Successfully downloaded a description for the payload '%s'"), *Id.ToString());	
+
+		// Now that we have the payload description we can start pulling the chunks from the Io service
+		// and reconstruct the final payload.
+		FUniqueBuffer Payload = FUniqueBuffer::Alloc(Response.MetaData.PayloadLength);
+		
+		int64 BytesLeft = Response.MetaData.PayloadLength;
+
+		FGraphEventArray Tasks;
+		Tasks.Reserve(Response.ChunkHashes.Num());
+
+		std::atomic<int32> NumFailedChunks(0);
+
+		bool bAllChunksPulled = true;
+		uint8* PayloadPtr = (uint8*)Payload.GetData();
+
+		for (const FString& HashString : Response.ChunkHashes)
+		{
+			checkf(BytesLeft > 0, TEXT("Ran out of buffer space before all payload chunks were read!"));
+
+			const int64 BytesToRead = FMath::Min(BytesLeft, Response.MetaData.ChunkLength);
+
+			auto Job = [this, PayloadPtr, BytesToRead, &Id, &HashString, &NumFailedChunks]()
+			{
+				if (!PullChunk(HashString, Id, PayloadPtr, BytesToRead))
+				{
+					NumFailedChunks++;
+				}
+			};
+
+#if UE_ENABLE_ASYNC_CHUNK_ACCESS
+			Tasks.Add(FFunctionGraphTask::CreateAndDispatchWhenReady(MoveTemp(Job)));
+#else
+			Job();
+#endif //UE_ENABLE_ASYNC_CHUNK_ACCESS
+
+			PayloadPtr += BytesToRead;
+			BytesLeft -= BytesToRead;
+		}
+
+		{
+			TRACE_CPUPROFILER_EVENT_SCOPE(FJupiterBackend::PullData::WaitOnChunks);
+			FTaskGraphInterface::Get().WaitUntilTasksComplete(Tasks);
+		}
+
+		if (NumFailedChunks == 0)
+		{
+			UE_LOG(LogVirtualization, Verbose, TEXT("Successfully downloaded all chunks for the payload '%s'"), *Id.ToString());	
+			return FCompressedBuffer::FromCompressed(Payload.MoveToShared());
+		}
+		else
+		{
+			UE_LOG(LogVirtualization, Error, TEXT("Failed to download %d chunks for the payload '%s'"), NumFailedChunks.load(), *Id.ToString());
+			return FCompressedBuffer();
+		}
+	}
+
+	virtual FString GetDebugString() const override
+	{
+		return FString(TEXT("Jupiter"));
+	}
+
+	bool IsUsingLocalHost() const
+	{
+		return HostAddress.StartsWith(TEXT("http://localhost"));
+	}
+
+	// TODO: Pretty much the same code as in FHttpDerivedDataBackend, another candidate for code sharing
+	bool IsServiceReady()
+	{
+		TRACE_CPUPROFILER_EVENT_SCOPE(FJupiterBackend::IsServiceReady);
+
+		using namespace Utility;
+
+		FRequest Request(*HostAddress, nullptr, false);
+		FRequest::Result Result = Request.PerformBlockingDownload(TEXT("health/ready"), nullptr);
+
+		if (Result == FRequest::Success && FRequest::IsSuccessfulResponse(Request.GetResponseCode()))
+		{
+			UE_LOG(LogVirtualization, Log, TEXT("Jupiter service status: '%s'."), *Request.GetResponseAsString());
+			return true;
+		}
+		else
+		{
+			UE_LOG(LogVirtualization, Error, TEXT("Unable to reach Jupiter service at '%s'. Status: %d . Response: '%s'"), *HostAddress, Request.GetResponseCode(), *Request.GetResponseAsString());
+			return false;
+		}	
+	}
+
+	// TODO: Pretty much the same code as in FHttpDerivedDataBackend, another candidate for code sharing
+	bool AcquireAccessToken()
+	{
+		TRACE_CPUPROFILER_EVENT_SCOPE(FJupiterBackend::AcquireAccessToken);
+
+		using namespace Utility;
+
+		if (IsUsingLocalHost())
+		{
+			UE_LOG(LogVirtualization, Log, TEXT("Connecting to a local host '%s', so skipping authorization"), *HostAddress);
+			return true;
+		}
+
+		// Avoid spamming the this if the service is down
+		if (FailedLoginAttempts > UE_MIRAGE_MAX_FAILED_LOGIN_ATTEMPTS)
+		{
+			return false;
+		}
+
+		ensureMsgf(OAuthProvider.StartsWith(TEXT("http://")) || OAuthProvider.StartsWith(TEXT("https://")),
+			TEXT("The OAuth provider %s is not valid. Needs to be a fully qualified url."),
+			*OAuthProvider
+		);
+
+		// In case many requests wants to update the token at the same time
+		// get the current serial while we wait to take the CS.
+		const uint32 WantsToUpdateTokenSerial = AccessToken.IsValid() ? AccessToken->GetSerial() : 0u;
+
+		{
+			FScopeLock Lock(&AccessCs);
+
+			// Check if someone has beaten us to update the token, then it 
+			// should now be valid.
+			if (AccessToken.IsValid() && AccessToken->GetSerial() > WantsToUpdateTokenSerial)
+			{
+				return true;
+			}
+
+			const uint32 SchemeEnd = OAuthProvider.Find(TEXT("://")) + 3;
+			const uint32 DomainEnd = OAuthProvider.Find(TEXT("/"), ESearchCase::CaseSensitive, ESearchDir::FromStart, SchemeEnd);
+			FString AuthDomain(DomainEnd, *OAuthProvider);
+			FString Uri(*OAuthProvider + DomainEnd + 1);
+
+			FRequest Request(*AuthDomain, nullptr, false);
+
+			// If contents of the secret string is a file path, resolve and read form data.
+			if (OAuthSecret.StartsWith(TEXT("\\\\")))
+			{
+				FString SecretFileContents;
+				if (FFileHelper::LoadFileToString(SecretFileContents, *OAuthSecret))
+				{
+					// Overwrite the filepath with the actual content.
+					OAuthSecret = SecretFileContents;
+				}
+				else
+				{
+					UE_LOG(LogVirtualization, Warning, TEXT("Failed to read OAuth form data file (%s)."), *OAuthSecret);
+					return false;
+				}
+			}
+
+			FString OAuthFormData = FString::Printf(
+				TEXT("client_id=%s&scope=cache_access&grant_type=client_credentials&client_secret=%s"),
+				*OAuthClientId,
+				*OAuthSecret
+			);
+
+			TArray<uint8> FormData;
+			auto OAuthFormDataUTF8 = FTCHARToUTF8(*OAuthFormData);
+			FormData.Append((uint8*)OAuthFormDataUTF8.Get(), OAuthFormDataUTF8.Length());
+
+			FRequest::Result Result = Request.PerformBlockingUpload<FRequest::Post>(*Uri, MakeArrayView(FormData));
+
+			if (Result == FRequest::Success && Request.GetResponseCode() == 200)
+			{
+				TSharedPtr<FJsonObject> ResponseObject = Request.GetResponseAsJsonObject();
+				if (ResponseObject)
+				{
+					FString AccessTokenString;
+					int32 ExpiryTimeSeconds = 0;
+					int32 CurrentTimeSeconds = int32(FPlatformTime::ToSeconds(FPlatformTime::Cycles()));
+
+					if (ResponseObject->TryGetStringField(TEXT("access_token"), AccessTokenString) &&
+						ResponseObject->TryGetNumberField(TEXT("expires_in"), ExpiryTimeSeconds))
+					{
+						if (!AccessToken)
+						{
+							AccessToken = MakeUnique<FAccessToken>();
+						}
+						AccessToken->SetHeader(*AccessTokenString);
+						UE_LOG(LogVirtualization, Log, TEXT("Logged in to HTTP DDC services. Expires in %d seconds."), ExpiryTimeSeconds);
+
+						//Schedule a refresh of the token ahead of expiry time (this will not work in commandlets)
+						if (!IsRunningCommandlet())
+						{
+							FTicker::GetCoreTicker().AddTicker(FTickerDelegate::CreateLambda(
+								[this](float DeltaTime)
+								{
+									this->AcquireAccessToken();
+									return false;
+								}
+							), ExpiryTimeSeconds - 20.0f);
+						}
+						// Reset failed login attempts, the service is indeed alive.
+						FailedLoginAttempts = 0;
+						return true;
+					}
+				}
+			}
+			else
+			{
+				UE_LOG(LogVirtualization, Warning, TEXT("Failed to log in to HTTP services. Server responed with code %d."), Request.GetResponseCode());
+				FailedLoginAttempts++;
+			}
+		}
+		return false;
+	}
+
+	/** 
+	 * Request the status of the service that we are connected to and make sure that it supports the 
+	 * feature set we need and meets our minimum version requirements.
+	 */
+	bool ValidateServiceVersion()
+	{
+		TRACE_CPUPROFILER_EVENT_SCOPE(FJupiterBackend::ValidateServiceVersion);
+
+		using namespace Utility;
+
+		FJupiterServiceStatus Response;
+
+		int64 ResponseCode = 0;
+		uint32 Attempts = 0;
+		while (ResponseCode != 200 && ++Attempts <= UE_MIRAGE_MAX_ATTEMPTS)
+		{
+			// We create the FRequest ourselves since the request pool does not yet exist
+			FRequest Request(*HostAddress, AccessToken.Get(), true);
+			FRequest::Result Result = Request.PerformBlockingDownload(TEXT("api/v1/status"));
+
+			ResponseCode = Request.GetResponseCode();
+
+			// Request was successful, make sure we got all the expected data.
+			if (ResponseCode == 200)
+			{
+				if (!Response.FromJson(Request.GetResponseAsJsonObject()))
+				{
+					UE_LOG(LogVirtualization, Error, TEXT("The response to 'api/v1/status' GET did not contain valid data!"));
+					return false;
+				}
+			}
+			else if (!ShouldRetryOnError(ResponseCode))
+			{
+				UE_LOG(LogVirtualization, Error, TEXT("Failed with error code '%d' to access the services status"), ResponseCode);
+				return false;
+			}
+		}
+
+		if (ResponseCode != 200)
+		{
+			UE_LOG(LogVirtualization, Error, TEXT("Failed '%d' attempts to access the services status (last error code '%d')"), UE_MIRAGE_MAX_ATTEMPTS, ResponseCode);
+			return false;
+		}
+
+		// Check version number
+		if (!Response.DoesHaveValidVersion(UE_MIRAGE_JUPITER_MIN_MAJOR_VER, UE_MIRAGE_JUPITER_MIN_MINOR_VER, UE_MIRAGE_JUPITER_MIN_PATCH_VER))
+		{
+			return false;
+		}
+
+		if (!Response.SupportsCapability(TEXT("ddc")))
+		{
+			UE_LOG(LogVirtualization, Error, TEXT("Jupiter service does not support Europa (ddc) capability"));
+			return false;
+		}
+
+		Response.LogStatusInfo();
+
+		return true;
+	}
+
+	// TODO: Pretty much the same code as in FHttpDerivedDataBackend, another candidate for code sharing
+	bool ShouldRetryOnError(int64 ResponseCode)
+	{
+		// Access token might have expired, request a new token and try again.
+		if (ResponseCode == 401 && AcquireAccessToken())
+		{
+			return true;
+		}
+
+		// Too many requests, make a new attempt
+		if (ResponseCode == 429)
+		{
+			return true;
+		}
+
+		// Gateway timeout, it will most likely work if we try again
+		if (ResponseCode == 504)
+		{
+			return true;
+		}
+
+		return false;
+	}
+
+	bool PostChunk(const TArrayView<const uint8>& ChunkData, const FPayloadId& PayloadId, FString& OutHashAsString)
+	{
+		TRACE_CPUPROFILER_EVENT_SCOPE(FJupiterBackend::PostChunk);
+
+		using namespace Utility;
+
+		const FIoHash ChunkHash = FIoHash::HashBuffer(ChunkData.GetData(), ChunkData.Num());
+		OutHashAsString = LexToString(ChunkHash);
+
+#if UE_CHECK_FOR_EXISTING_CHUNKS
+		if (DoesChunkExist(HashAsString))
+		{
+			UE_LOG(LogVirtualization, Verbose, TEXT("Jupiter already has a copy of the chunk '%s' for payload '%s'"), *HashAsString, *PayloadId.ToString());
+			return true;
+		}
+#endif // UE_CHECK_FOR_EXISTING_CHUNKS
+
+		FScopedRequestPtr Request(RequestPool.Get());
+		// TODO: Another candidate for code sharing
+		Request->SetHeader(TEXT("X-Jupiter-IoHash"), *OutHashAsString);		
+
+		TStringBuilder<256> Uri;
+		Uri.Appendf(TEXT("api/v1/s/%s/%s"), *Namespace, *OutHashAsString);
+
+		// Retry request until we get an accepted response or exhaust allowed number of attempts.
+		uint32 Attempts = 0;
+		while (++Attempts <= UE_MIRAGE_MAX_ATTEMPTS)
+		{
+			if (Request.IsValid())
+			{
+				const FRequest::Result Result = Request->PerformBlockingUpload<FRequest::Put>(*Uri, ChunkData);
+				const int64 ResponseCode = Request->GetResponseCode();
+
+				if (ResponseCode == 200)
+				{
+					UE_LOG(LogVirtualization, Verbose, TEXT("Successfully uploaded a chunk '%s'for payload '%s'"), *OutHashAsString, *PayloadId.ToString());	
+					return true;
+				}
+
+				if (!ShouldRetryOnError(ResponseCode))
+				{
+					return false;
+				}
+			}
+		}
+
+		return false;
+	}
+
+	bool PullChunk(const FString& Hash, const FPayloadId& PayloadId, uint8* DataPtr, int64 BufferSize)
+	{
+		TRACE_CPUPROFILER_EVENT_SCOPE(FJupiterBackend::PullChunk);
+
+		using namespace Utility;
+
+		FScopedRequestPtr Request(RequestPool.Get());
+
+		TStringBuilder<256> Uri;
+		Uri.Appendf(TEXT("api/v1/s/%s/%s"), *Namespace, *Hash);
+
+		// Retry request until we get an accepted response or exhaust allowed number of attempts.
+		uint32 Attempts = 0;
+		while (++Attempts <= UE_MIRAGE_MAX_ATTEMPTS)
+		{
+			if (Request.IsValid())
+			{
+				const FRequest::Result Result = Request->PerformBlockingDownload(*Uri, FMutableMemoryView(DataPtr, BufferSize));
+
+				if (Result != FRequest::Success)
+				{
+					UE_LOG(LogVirtualization, Error, TEXT("Attempting to GET a payload chunk '%s' for payload '%s' failed due to an internal Curl error"), *Hash, *PayloadId.ToString());	
+					return false;
+				}
+
+				const int64 ResponseCode = Request->GetResponseCode();
+
+				if (ResponseCode == 200)
+				{
+					UE_LOG(LogVirtualization, Verbose, TEXT("Successfully downloaded a payload chunk '%s' for payload '%s'"), *Hash, *PayloadId.ToString());	
+					return true;
+				}
+
+				if (!ShouldRetryOnError(ResponseCode))
+				{
+					UE_LOG(LogVirtualization, Error, TEXT("Attempting to GET a payload chunk '%s' for payload '%s' failed with http response: %" INT64_FMT), *Hash, *PayloadId.ToString(), ResponseCode);	
+					return false;
+				}
+			}
+		}
+
+		UE_LOG(LogVirtualization, Error, TEXT("Attempting to GET a payload chunk '%s' for payload '%s' failed all '%d' attempts"), *Hash, *PayloadId.ToString(), UE_MIRAGE_MAX_ATTEMPTS);	
+		return false;
+	}
+
+	bool DoesPayloadExist(const FPayloadId& Id)
+	{
+		using namespace Utility;
+
+		TRACE_CPUPROFILER_EVENT_SCOPE(FJupiterBackend::DoesPayloadExist);
+
+		// Note that the ddc end point is used by both ddc and mirage
+		TStringBuilder<256> Uri;
+		Uri.Appendf(TEXT("api/v1/c/ddc/%s/%s/%s"), *Namespace, *Bucket, *Id.ToString());
+		
+		// Retry request until we get an accepted response or exhaust allowed number of attempts.
+		uint32 Attempts = 0;
+		while (++Attempts <= UE_MIRAGE_MAX_ATTEMPTS)
+		{
+			FScopedRequestPtr Request(RequestPool.Get());
+			if (Request.IsValid())
+			{
+				const FRequest::Result Result = Request->PerformBlockingQuery<FRequest::Head>(*Uri);
+				const int64 ResponseCode = Request->GetResponseCode();
+
+				if (FRequest::IsSuccessfulResponse(ResponseCode))
+				{
+					return true;
+				}
+				else if (ResponseCode == 400)
+				{
+					return false;
+				}
+
+				if (!ShouldRetryOnError(ResponseCode))
+				{
+					return false;
+				}
+			}
+		}
+
+		return false;
+	}
+
+	bool DoesChunkExist(const FString& Hash)
+	{
+		using namespace Utility;
+
+		TRACE_CPUPROFILER_EVENT_SCOPE(FJupiterBackend::DoesChunkExist);
+
+		TStringBuilder<256> Uri;
+		Uri.Appendf(TEXT("api/v1/s/%s/%s"), *Namespace, *Hash);
+
+		// Retry request until we get an accepted response or exhaust allowed number of attempts.
+		uint32 Attempts = 0;
+		while (++Attempts <= UE_MIRAGE_MAX_ATTEMPTS)
+		{
+			FScopedRequestPtr Request(RequestPool.Get());
+			if (Request.IsValid())
+			{
+				const FRequest::Result Result = Request->PerformBlockingQuery<FRequest::Head>(*Uri);
+				const int64 ResponseCode = Request->GetResponseCode();
+
+				if (FRequest::IsSuccessfulResponse(ResponseCode))
+				{
+					return true;
+				}
+				else if (ResponseCode == 404)
+				{
+					return false;
+				}
+
+				if (!ShouldRetryOnError(ResponseCode))
+				{
+					return false;
+				}
+			}
+		}
+
+		return false;
+	}
+	
+	/** Address of the Jupiter service*/
+	FString HostAddress;
+	/** Namespace to connect to */
+	FString Namespace;
+	/** Europa allows us to organize the payloads by bucket. Currently this is not exposed and just set to 'default' */
+	FString Bucket;
+
+	/** The max size of each payload chunk */
+	uint64 ChunkSize;
+
+	/** Url of the OAuth authorization server */
+	FString OAuthProvider;
+	/**  Public identifier for use with the OAuth authorization server */
+	FString OAuthClientId;
+	/** Password for the OAuthClientId */
+	FString OAuthSecret;
+
+	/** The pool of FRequest objects that can be recycled */
+	TUniquePtr<Utility::FRequestPool> RequestPool;
+
+	/** Critical section used to protect the creation of new access tokens */
+	FCriticalSection AccessCs;
+	/** The access token used with service authorization */
+	TUniquePtr<Utility::FAccessToken> AccessToken;
+	/** Count how many times a login has failed since the last successful login */
+	uint32 FailedLoginAttempts;
+};
+
+UE_REGISTER_VIRTUALIZATION_BACKEND_FACTORY(FJupiterBackend, Jupiter);
+
+} // namespace Virtualization
+} // namespace UE
+
+#endif //WITH_MIRAGE_JUPITER_BACKEND
diff --git a/Engine/Source/Developer/Virtualization/Virtualization.Build.cs b/Engine/Source/Developer/Virtualization/Virtualization.Build.cs
new file mode 100644
--- /dev/null
+++ b/Engine/Source/Developer/Virtualization/Virtualization.Build.cs
@@ -0,0 +1,21 @@
+// Copyright Epic Games, Inc. All Rights Reserved.
+
+using UnrealBuildTool;
+using System.IO;
+
+public class Virtualization : ModuleRules
+{
+	public Virtualization(ReadOnlyTargetRules Target) : base(Target)
+	{
+		PrivateDependencyModuleNames.Add("Core");
+
+		// Remove this when we move the rest of the Virtualization code to this module
+		PrivateDependencyModuleNames.Add("CoreUObject"); 
+
+		// Dependencies for the Jupiter service backend
+		PrivateDependencyModuleNames.AddRange(new string[] { "SSL", "Json", "SourceControl" });
+
+		AddEngineThirdPartyPrivateStaticDependencies(Target, "libcurl");
+		AddEngineThirdPartyPrivateStaticDependencies(Target, "OpenSSL");
+	}
+}
diff --git a/Engine/Source/Runtime/Core/Core.Build.cs b/Engine/Source/Runtime/Core/Core.Build.cs
--- a/Engine/Source/Runtime/Core/Core.Build.cs
+++ b/Engine/Source/Runtime/Core/Core.Build.cs
@@ -163,6 +163,7 @@ public class Core : ModuleRules
 		if (Target.bCompileAgainstEngine && !Target.bBuildRequiresCookedData)
 		{
 			DynamicallyLoadedModuleNames.AddRange(new string[] { "DerivedDataCache" });
+			DynamicallyLoadedModuleNames.AddRange(new string[] { "Virtualization" });
 		}
 
 		if (Target.Platform.IsInGroup(UnrealPlatformGroup.Desktop))
diff --git a/Engine/Source/Runtime/Core/Private/Misc/ConfigCacheIni.cpp b/Engine/Source/Runtime/Core/Private/Misc/ConfigCacheIni.cpp
--- a/Engine/Source/Runtime/Core/Private/Misc/ConfigCacheIni.cpp
+++ b/Engine/Source/Runtime/Core/Private/Misc/ConfigCacheIni.cpp
@@ -84,6 +84,13 @@ struct FConfigExpansion
 {
 	template<int N>
 	FConfigExpansion(const TCHAR(&Var)[N], FString&& Val)
+		: Variable(Var)
+		, Value(MoveTemp(Val))
+		, VariableLen(N - 1)
+	{}
+
+	template<int N>
+	FConfigExpansion(const TCHAR(&Var)[N], const FString& Val)
 		: Variable(Var)
 		, Value(Val)
 		, VariableLen(N - 1)
@@ -112,6 +119,7 @@ static const FConfigExpansion* MatchExpansions(const TCHAR* PotentialVariable)
 		FConfigExpansion(TEXT("%ENGINEUSERDIR%"), FPaths::EngineUserDir()),
 		FConfigExpansion(TEXT("%ENGINEVERSIONAGNOSTICUSERDIR%"), FPaths::EngineVersionAgnosticUserDir()),
 		FConfigExpansion(TEXT("%APPSETTINGSDIR%"), GetApplicationSettingsDirNormalized()),
+		FConfigExpansion(TEXT("%GAMESAVEDDIR%"), FPaths::ProjectSavedDir()),
 	};
 
 	for (const FConfigExpansion& Expansion : Expansions)
diff --git a/Engine/Source/Runtime/Core/Private/Misc/CoreGlobals.cpp b/Engine/Source/Runtime/Core/Private/Misc/CoreGlobals.cpp
--- a/Engine/Source/Runtime/Core/Private/Misc/CoreGlobals.cpp
+++ b/Engine/Source/Runtime/Core/Private/Misc/CoreGlobals.cpp
@@ -699,6 +699,7 @@ DEFINE_LOG_CATEGORY(LogNetSerialization);
 DEFINE_LOG_CATEGORY(LogMemory);
 DEFINE_LOG_CATEGORY(LogProfilingDebugging);
 DEFINE_LOG_CATEGORY(LogTemp);
+DEFINE_LOG_CATEGORY(LogVirtualization);
 
 // need another layer of macro to help using a define in a define
 #define DEFINE_LOG_CATEGORY_HELPER(A) DEFINE_LOG_CATEGORY(A)
diff --git a/Engine/Source/Runtime/Core/Private/Tests/Virtualization/PayloadIdTests.cpp b/Engine/Source/Runtime/Core/Private/Tests/Virtualization/PayloadIdTests.cpp
new file mode 100644
--- /dev/null
+++ b/Engine/Source/Runtime/Core/Private/Tests/Virtualization/PayloadIdTests.cpp
@@ -0,0 +1,329 @@
+// Copyright Epic Games, Inc. All Rights Reserved.
+
+#include "Memory/SharedBuffer.h"
+#include "Misc/StringBuilder.h"
+#include "Serialization/MemoryReader.h"
+#include "Serialization/MemoryWriter.h"
+#include "Virtualization/PayloadId.h"
+
+#include "Misc/AutomationTest.h"
+
+#if WITH_DEV_AUTOMATION_TESTS
+
+namespace UE::Virtualization
+{
+	
+constexpr const uint32 TestFlags = EAutomationTestFlags::ApplicationContextMask | EAutomationTestFlags::SmokeFilter;
+
+/** This test ensures that ::IsValid returns the correct value dending on how the FPayloadId was created */
+IMPLEMENT_SIMPLE_AUTOMATION_TEST(FPayloadIdTestValidity, TEXT("System.Core.Virtualization.PayloadId.Validity"), TestFlags)
+bool FPayloadIdTestValidity::RunTest(const FString& Parameters)
+{
+	// Create a test buffer for us to use in the tests
+	const uint32 TestDataLength = 32;
+	uint32 TestDataBuffer[TestDataLength];
+	for (uint32 Index = 0; Index < TestDataLength; ++Index)
+	{
+		TestDataBuffer[Index] = Index;
+	}
+
+	// Default constructor
+	{
+		const FPayloadId EmptyId;
+		TestFalse(TEXT("An empty FPayloadId ::IsValid()"), EmptyId.IsValid());
+	}
+
+	// FIoHash
+	{
+		const FIoHash EmptyHash;
+		const FPayloadId EmptyHashId(EmptyHash);
+		TestFalse(TEXT("A FPayloadId from an empty FIoHash ::IsValid()"), EmptyHashId.IsValid());
+
+		const FIoHash EmptyBufferHash = FIoHash::HashBuffer(nullptr, 0);
+		const FPayloadId EmptyBufferHashId(EmptyBufferHash);
+		TestTrue(TEXT("A FPayloadId from a FIoHash of an empty buffer ::IsValid()"), EmptyBufferHashId.IsValid());
+
+		const FIoHash BufferHash = FIoHash::HashBuffer(TestDataBuffer, TestDataLength * sizeof(uint32));
+		const FPayloadId BufferHashId(BufferHash);
+		TestTrue(TEXT("A FPayloadId from a FIoHash of a buffer ::IsValid()"), BufferHashId.IsValid());
+	}
+
+	// FGuid
+	{
+		const FGuid EmptyGuid;
+		const FPayloadId EmptyGuidId(EmptyGuid);
+		TestFalse(TEXT("A FPayloadId from an invalid FGuid ::IsValid()"), EmptyGuidId.IsValid());
+
+		const FGuid ValidGuid = FGuid::NewGuid();
+		const FPayloadId ValidGuidId(ValidGuid);
+		TestTrue(TEXT("A FPayloadId from a valid FGuid ::IsValid()"), ValidGuidId.IsValid());
+	}
+
+	// FSharedBuffer
+	{
+		FSharedBuffer NullBuffer;
+		const FPayloadId NullBufferId(NullBuffer);
+		TestFalse(TEXT("A FPayloadId from a null FSharedBuffer ::IsValid()"), NullBufferId.IsValid());
+
+		FSharedBuffer EmptyBuffer = FSharedBuffer::MakeView(nullptr, 0);
+		const FPayloadId EmptyBufferId(EmptyBuffer);
+		TestFalse(TEXT("A FPayloadId from an empty FSharedBuffer ::IsValid()"), EmptyBufferId.IsValid());
+
+		FSharedBuffer ValidBuffer = FSharedBuffer::MakeView(TestDataBuffer, TestDataLength);
+		const FPayloadId ValidBufferId(ValidBuffer);
+		TestTrue(TEXT("A FPayloadId from a valid FSharedBuffer ::IsValid()"), ValidBufferId.IsValid());
+	}
+
+	return true;
+}
+
+/** This test ensures that the == operator is functioning correctly */
+IMPLEMENT_SIMPLE_AUTOMATION_TEST(FPayloadIdTestEquality, TEXT("System.Core.Virtualization.PayloadId.Equality"), TestFlags)
+bool FPayloadIdTestEquality::RunTest(const FString& Parameters)
+{
+	// Create a test buffer for us to use in the tests
+	const uint32 TestDataLength = 32;
+	const uint32 TestDataNumBytes = TestDataLength * sizeof(uint32);
+
+	uint32 TestDataBuffer0[TestDataLength];
+	uint32 TestDataBuffer1[TestDataLength];
+	uint32 TestDataBuffer3[TestDataLength];
+	
+	for (uint32 Index = 0; Index < TestDataLength; ++Index)
+	{
+		TestDataBuffer0[Index] = Index;
+		TestDataBuffer1[Index] = Index * 2; // Note the first two entries will be the same as 
+											// the other buffers
+		TestDataBuffer3[Index] = Index;
+	}
+
+	const FPayloadId EmptyId0;
+	const FPayloadId EmptyId1;
+
+	const FPayloadId BufferId0(FIoHash::HashBuffer(TestDataBuffer0, TestDataNumBytes));
+	const FPayloadId BufferId1(FIoHash::HashBuffer(TestDataBuffer1, TestDataNumBytes));
+	const FPayloadId BufferId2(FIoHash::HashBuffer(TestDataBuffer3, TestDataNumBytes));
+
+	TestTrue(TEXT("Two empty FPayloadId are equal"), EmptyId0 == EmptyId1);
+	TestFalse(TEXT("An empty FPayloadId and a valid FPayloadId are equal"), EmptyId0 == BufferId0);
+	TestFalse(TEXT("Two FPayloadId from different buffers with different values are equal"), BufferId0 == BufferId1);
+	TestTrue(TEXT("Two FPayloadId from different but identical buffers are equal"), BufferId0 == BufferId2);
+
+	return true;
+}
+
+/** This test ensures that the != operator is functioning correctly */
+IMPLEMENT_SIMPLE_AUTOMATION_TEST(FPayloadIdTestInequality, TEXT("System.Core.Virtualization.PayloadId.Inequality"), TestFlags)
+bool FPayloadIdTestInequality::RunTest(const FString& Parameters)
+{
+	// Create a test buffer for us to use in the tests
+	const uint32 TestDataLength = 32;
+	const uint32 TestDataNumBytes = TestDataLength * sizeof(uint32);
+
+	uint32 TestDataBuffer0[TestDataLength];
+	uint32 TestDataBuffer1[TestDataLength];
+	uint32 TestDataBuffer3[TestDataLength];
+
+	for (uint32 Index = 0; Index < TestDataLength; ++Index)
+	{
+		TestDataBuffer0[Index] = Index;
+		TestDataBuffer1[Index] = Index * 2; // Note the first two entries will be the same as 
+											// the other buffers
+		TestDataBuffer3[Index] = Index;
+	}
+
+	const FPayloadId EmptyId0;
+	const FPayloadId EmptyId1;
+
+	const FPayloadId BufferId0(FIoHash::HashBuffer(TestDataBuffer0, TestDataNumBytes));
+	const FPayloadId BufferId1(FIoHash::HashBuffer(TestDataBuffer1, TestDataNumBytes));
+	const FPayloadId BufferId2(FIoHash::HashBuffer(TestDataBuffer3, TestDataNumBytes));
+
+	TestFalse(TEXT("Two empty FPayloadId are not equal"), EmptyId0 != EmptyId1);
+	TestTrue(TEXT("An empty FPayloadId and a valid FPayloadId are not equal"), EmptyId0 != BufferId0);
+	TestTrue(TEXT("Two FPayloadId from different buffers with different values are not equal"), BufferId0 != BufferId1);
+	TestFalse(TEXT("Two FPayloadId from different but identical buffers are not equal"), BufferId0 != BufferId2);
+
+	return true;
+}
+
+IMPLEMENT_SIMPLE_AUTOMATION_TEST(FPayloadIdTestSerialization, TEXT("System.Core.Virtualization.PayloadId.Serialization"), TestFlags)
+bool FPayloadIdTestSerialization::RunTest(const FString& Parameters)
+{
+	auto Serialize = [](FPayloadId InSrc, FPayloadId& OutDst) 
+	{
+		TArray<uint8> MemoryBuffer;
+		const bool bIsArPersistent = true;
+
+		{
+			FMemoryWriter WriterAr(MemoryBuffer, bIsArPersistent);
+			WriterAr << InSrc;
+		}
+
+		{
+			FMemoryReader ReaderAr(MemoryBuffer, bIsArPersistent);
+			ReaderAr << OutDst;
+		}	
+	};
+
+	// Serialize an empty FPayloadId
+	{
+		FPayloadId Source;
+		FPayloadId Result;
+
+		Serialize(Source, Result);
+
+		TestTrue(TEXT("Empty FPayloadId serializes to an empty FPayloadId"), Source == Result);
+		TestTrue(TEXT("A serialized empty FPayloadId should be invalid"), Result.IsValid() == false);
+	}
+
+	// Serialize a valid FPayloadId
+	{
+		FPayloadId Source(FIoHash(TEXT("73cdaedfeff72f606fc1e73c9751a8418275da58")));
+		FPayloadId Result;
+
+		Serialize(Source, Result);
+
+		TestTrue(TEXT("Valid FPayloadId serializes to a valid FPayloadId"), Source == Result);
+		TestTrue(TEXT("A serialized FPayloadId should be valid"), Result.IsValid() == true);
+	}
+	
+	return true;
+}
+
+IMPLEMENT_SIMPLE_AUTOMATION_TEST(FPayloadIdTestTypeHash, TEXT("System.Core.Virtualization.PayloadId.TypeHash"), TestFlags)
+bool FPayloadIdTestTypeHash::RunTest(const FString& Parameters)
+{
+	// Create a test buffer for us to use in the tests
+	const uint32 TestDataLength = 32;
+	const uint32 TestDataNumBytes = TestDataLength * sizeof(uint32);
+
+	uint32 TestDataBuffer0[TestDataLength];
+	uint32 TestDataBuffer1[TestDataLength];
+	uint32 TestDataBuffer3[TestDataLength];
+
+	for (uint32 Index = 0; Index < TestDataLength; ++Index)
+	{
+		TestDataBuffer0[Index] = Index;
+		TestDataBuffer1[Index] = Index * 2; // Note the first two entries will be the same as 
+											// the other buffers
+		TestDataBuffer3[Index] = Index;
+	}
+
+	// Set up the FPayloadId's for the test
+	const FPayloadId EmptyId0;
+	const FPayloadId EmptyId1;
+
+	const FPayloadId BufferId0(FIoHash::HashBuffer(TestDataBuffer0, TestDataNumBytes));
+	const FPayloadId BufferId1(FIoHash::HashBuffer(TestDataBuffer1, TestDataNumBytes));
+	const FPayloadId BufferId2(FIoHash::HashBuffer(TestDataBuffer3, TestDataNumBytes));
+
+	// Test direct use of GetTypeHash
+	{
+		TestEqual(TEXT("Two empty FPayloadId should have the same hashes"), GetTypeHash(EmptyId1), GetTypeHash(EmptyId0));
+		TestNotEqual(TEXT("An empty FPayloadId and a valid FPayloadId should have different hashes"), GetTypeHash(BufferId0), GetTypeHash(EmptyId0));
+		TestNotEqual(TEXT("An empty FPayloadId and a valid FPayloadId should have different hashes"), GetTypeHash(BufferId1), GetTypeHash(EmptyId0));
+		TestNotEqual(TEXT("An empty FPayloadId and a valid FPayloadId should have different hashes"), GetTypeHash(BufferId2), GetTypeHash(EmptyId0));
+		TestNotEqual(TEXT("Two valid FPayloadId created from different buffers should have different hashes"), GetTypeHash(BufferId0), GetTypeHash(BufferId1));
+		TestEqual(TEXT("Two valid FPayloadId created from binary equivillant buffers should have the same hashes"), GetTypeHash(BufferId0), GetTypeHash(BufferId2));
+	}
+
+	// Test the indirect use of GetTypeHash in practice via TMap
+	{
+#define TEST_MAP_ENTRY(Key, ExpectedValue) { uint32* Value = PayloadIdMap.Find(Key); \
+											if (Value != nullptr) { TestEqual(TEXT("Data stored for key: " #Key), *Value, ExpectedValue); } else \
+											{ AddError(TEXT("Unable to find entry for key: " #Key)); }}
+
+		TMap<FPayloadId, uint32> PayloadIdMap;
+		
+		PayloadIdMap.Add(EmptyId0, 0);
+		
+		// Test that adding EmptyId0 created a single entry, and that since EmptyId0 and EmptyId1 
+		// have the same type hash that we can access the value by both EmptyId0 and EmptyId1
+		TestEqual(TEXT("Map Count"), PayloadIdMap.Num(), 1);
+		TEST_MAP_ENTRY(EmptyId0, 0);
+		TEST_MAP_ENTRY(EmptyId1, 0);
+
+		PayloadIdMap.Add(EmptyId1, 1); // Should replace existing entry for EmptyId0!
+
+		// Test that we still only have a single entry and that adding EmptyId1 replaced the value
+		// in the entry for EmptyId0
+		TestEqual(TEXT("Map Count"), PayloadIdMap.Num(), 1);
+		TEST_MAP_ENTRY(EmptyId0, 1);
+		TEST_MAP_ENTRY(EmptyId1, 1);
+
+		PayloadIdMap.Add(BufferId0, 2);
+		PayloadIdMap.Add(BufferId1, 3);
+
+		// Test that adding the two new entries worked and that we can access the value of 
+		// BufferId0 by BufferId2 since they produce the same type hash
+		TestEqual(TEXT("Map Count"), PayloadIdMap.Num(), 3);
+		TEST_MAP_ENTRY(EmptyId0, 1);
+		TEST_MAP_ENTRY(EmptyId1, 1);
+		TEST_MAP_ENTRY(BufferId0, 2);
+		TEST_MAP_ENTRY(BufferId1, 3);
+		TEST_MAP_ENTRY(BufferId2, 2);
+
+		PayloadIdMap.Add(BufferId2, 4); // Should replace existing entry for BufferId0!
+
+		// Test that adding BufferId2 did not add a new entry but correctly replaced the
+		// existing entry for BufferId0
+		TestEqual(TEXT("Map Count"), PayloadIdMap.Num(), 3);
+		TEST_MAP_ENTRY(EmptyId0, 1);
+		TEST_MAP_ENTRY(EmptyId1, 1);
+		TEST_MAP_ENTRY(BufferId0, 4);
+		TEST_MAP_ENTRY(BufferId1, 3);
+		TEST_MAP_ENTRY(BufferId2, 4);
+
+#undef TEST_MAP_ENTRY
+	}
+
+
+	return true;
+}
+
+IMPLEMENT_SIMPLE_AUTOMATION_TEST(FPayloadIdTestStringBuilder, TEXT("System.Core.Virtualization.PayloadId.StringBuilder"), TestFlags)
+bool FPayloadIdTestStringBuilder::RunTest(const FString& Parameters)
+{
+	// Create a test buffer for us to use in the tests
+	const uint32 TestDataLength = 32;
+	const uint32 TestDataNumBytes = TestDataLength * sizeof(uint32);
+
+	uint32 TestDataBuffer0[TestDataLength];
+	uint32 TestDataBuffer1[TestDataLength];
+	uint32 TestDataBuffer3[TestDataLength];
+
+	for (uint32 Index = 0; Index < TestDataLength; ++Index)
+	{
+		TestDataBuffer0[Index] = Index;
+		TestDataBuffer1[Index] = Index * 2; // Note the first two entries will be the same as 
+											// the other buffers
+		TestDataBuffer3[Index] = Index;
+	}
+
+	// Set up the FPayloadId's for the test
+	const FPayloadId EmptyId;
+
+	const FPayloadId BufferId0(FIoHash::HashBuffer(TestDataBuffer0, TestDataNumBytes));
+	const FPayloadId BufferId1(FIoHash::HashBuffer(TestDataBuffer1, TestDataNumBytes));
+	const FPayloadId BufferId2(FIoHash::HashBuffer(TestDataBuffer3, TestDataNumBytes));
+
+	// Test again both TWideStringBuilder and TAnsiStringBuilder
+#define TEST_PAYLOAD_ID(BufferToTest) {	TWideStringBuilder<128> WideBuilder; WideBuilder << BufferToTest; \
+										TestEqual(TEXT("TWideStringBuilder << and ToString to produce the same result"), *WideBuilder, BufferToTest.ToString());\
+										TAnsiStringBuilder<128> AnsiBuilder; AnsiBuilder << BufferToTest; \
+										TestEqual(TEXT("TAnsiBuilder << and ToString to produce the same result"), *AnsiBuilder, BufferToTest.ToString()); }
+
+	TEST_PAYLOAD_ID(EmptyId);
+	TEST_PAYLOAD_ID(BufferId0);
+	TEST_PAYLOAD_ID(BufferId1);
+	TEST_PAYLOAD_ID(BufferId2);
+
+#undef TEST_PAYLOAD_ID
+
+	return true;
+}
+
+} // namespace UE::Virtualization
+
+#endif //WITH_DEV_AUTOMATION_TESTS
diff --git a/Engine/Source/Runtime/Core/Private/Virtualization/PayloadId.cpp b/Engine/Source/Runtime/Core/Private/Virtualization/PayloadId.cpp
new file mode 100644
--- /dev/null
+++ b/Engine/Source/Runtime/Core/Private/Virtualization/PayloadId.cpp
@@ -0,0 +1,121 @@
+// Copyright Epic Games, Inc. All Rights Reserved.
+
+#include "Virtualization/PayloadId.h"
+
+#include "Hash/Blake3.h"
+#include "Memory/SharedBuffer.h"
+#include "Misc/Guid.h"
+#include "Misc/StringBuilder.h"
+#include "Serialization/Archive.h"
+
+namespace UE
+{
+namespace Virtualization
+{
+
+FPayloadId::FPayloadId(const FIoHash& BlakeHash)
+	: Identifier(BlakeHash)
+{
+	bIsHashSet = !Identifier.IsZero();
+}
+
+FPayloadId::FPayloadId(const FSharedBuffer& Payload)
+{ 
+	// We consider not having a valid payload and having an empty payload to be the same.
+	if (Payload.GetSize() > 0)
+	{
+		Identifier = FIoHash::HashBuffer(Payload);
+		bIsHashSet = true;
+	}
+}
+
+// TODO: Currently only used for backwards compatibility, we might want to make this private and allow 
+// FVirtualizedUntypedBulkData::CreateFromBulkData to access it by being a friend.
+FPayloadId::FPayloadId(const FGuid& Guid)
+{
+	if (Guid.IsValid())
+	{
+		// Hash each element individually rather than making assumptions about
+		// the internal layout of FGuid and treating it as a contiguous buffer.
+		// Slightly slower, but safer.
+		FBlake3 Hash;
+
+		Hash.Update(&Guid[0], sizeof(uint32));
+		Hash.Update(&Guid[1], sizeof(uint32));
+		Hash.Update(&Guid[2], sizeof(uint32));
+		Hash.Update(&Guid[3], sizeof(uint32));
+
+		Identifier = FIoHash(Hash.Finalize());
+		bIsHashSet = true;
+	}
+}
+
+void FPayloadId::Reset()
+{
+	Identifier.Reset();
+	bIsHashSet = false;
+}
+
+FGuid FPayloadId::ToGuid() const
+{
+	if (IsValid())
+	{
+		// We use the first 16 bytes of the FIoHash to create the guid, there is
+		// no specific reason why these were chosen, we could take any pattern or combination
+		// of bytes.
+		uint32* HashBytes = (uint32*)Identifier.GetBytes();
+		return FGuid(HashBytes[0], HashBytes[1], HashBytes[2], HashBytes[3]);
+	}
+	else
+	{
+		return FGuid();
+	}
+}
+
+FString FPayloadId::ToString() const
+{
+	if (IsValid())
+	{	
+		TStringBuilder<65> HexStringBuilder;
+		HexStringBuilder << Identifier;
+
+		return HexStringBuilder.ToString();
+	}
+	else
+	{
+		return FString();
+	}
+}
+
+bool FPayloadId::operator == (const FPayloadId& InPayloadId) const
+{
+	// If ::bIsHashSet is different between the two payloads we know they are different.
+	if (bIsHashSet != InPayloadId.bIsHashSet)
+	{
+		return false;
+	}
+
+	// If ::bIsHashSet is false for both payload then we know they are the same.
+	if (!bIsHashSet)
+	{
+		return true;
+	}
+
+	// Compare the two hashes
+	return Identifier == InPayloadId.Identifier;
+}
+
+FArchive& operator<<(FArchive& Ar, FPayloadId& PayloadId)
+{
+	Ar << PayloadId.Identifier;
+
+	if (Ar.IsLoading())
+	{
+		PayloadId.bIsHashSet = !PayloadId.Identifier.IsZero();
+	}
+
+	return Ar;
+}
+
+} // namespace Virtualization
+} // namespace UE
diff --git a/Engine/Source/Runtime/Core/Public/CoreGlobals.h b/Engine/Source/Runtime/Core/Public/CoreGlobals.h
--- a/Engine/Source/Runtime/Core/Public/CoreGlobals.h
+++ b/Engine/Source/Runtime/Core/Public/CoreGlobals.h
@@ -41,6 +41,7 @@ CORE_API DECLARE_LOG_CATEGORY_EXTERN(LogLocalization, Error, All);
 CORE_API DECLARE_LOG_CATEGORY_EXTERN(LogLongPackageNames, Log, All);
 CORE_API DECLARE_LOG_CATEGORY_EXTERN(LogProcess, Log, All);
 CORE_API DECLARE_LOG_CATEGORY_EXTERN(LogLoad, Log, All);
+CORE_API DECLARE_LOG_CATEGORY_EXTERN(LogVirtualization, Log, All);
 
 // Temporary log category, generally you should not check things in that use this
 CORE_API DECLARE_LOG_CATEGORY_EXTERN(LogTemp, Log, All);
diff --git a/Engine/Source/Runtime/Core/Public/UObject/UE5CookerObjectVersion.h b/Engine/Source/Runtime/Core/Public/UObject/UE5CookerObjectVersion.h
--- a/Engine/Source/Runtime/Core/Public/UObject/UE5CookerObjectVersion.h
+++ b/Engine/Source/Runtime/Core/Public/UObject/UE5CookerObjectVersion.h
@@ -22,10 +22,10 @@ struct CORE_API FUE5CookerObjectVersion
 		// When UE_USE_VIRTUALBULKDATA is removed then they can be moved to the normal section
 
 #if UE_USE_VIRTUALBULKDATA
-		
 		// Switch FMeshDescriptionBulkData to use virtualized bulkdata
 		MeshDescriptionVirtualization,
-
+		// Switch FTextureSource to use virtualized bulkdata
+		TextureSourceVirtualization,
 #endif //UE_USE_VIRTUALBULKDATA
 
 		VersionPlusOne,
diff --git a/Engine/Source/Runtime/Core/Public/Virtualization/IVirtualizedData.h b/Engine/Source/Runtime/Core/Public/Virtualization/IVirtualizedData.h
new file mode 100644
--- /dev/null
+++ b/Engine/Source/Runtime/Core/Public/Virtualization/IVirtualizedData.h
@@ -0,0 +1,74 @@
+// Copyright Epic Games, Inc. All Rights Reserved.
+
+#pragma once
+
+#include "Async/Future.h"
+#include "Templates/Function.h"
+
+class FCompressedBuffer;
+class FSharedBuffer;
+
+namespace UE::Virtualization
+{
+	
+class FPayloadId;
+
+/**
+ * Base interface for virtualized data
+ */
+class IVirtualizedData
+{
+public:
+	
+	virtual ~IVirtualizedData() = default;
+
+public:
+	/**
+	 * Returns a unique identifier for the object itself.
+	 * This should only return a valid FGuid as long as the object owns a valid payload.
+	 * If an object with a valid payload, has that payload removed then it should start
+	 * returning an invalid FGuid instead.
+	 * Should that object be given a new payload it should then return the original 
+	 * identifier, there is no need to generate a new one.
+	 */
+	virtual FGuid GetIdentifier() const = 0;
+
+	/** Returns an unique identifier for the content of the payload. */
+	virtual const FPayloadId& GetPayloadId() const = 0;
+
+	/** Returns the size of the payload in bytes */
+	virtual int64 GetPayloadSize() const = 0;
+
+	/** Returns an immutable FSharedBuffer reference to the payload data. */
+	virtual TFuture<FSharedBuffer> GetPayload() const = 0;
+
+	/**
+	  * Returns an immutable FCompressedBuffer reference to the payload data.
+	  *
+	  * Note that depending on the internal storage formats, the payload might not actually be compressed, but that
+	  * will be handled by the FCompressedBuffer interface. Call FCompressedBuffer::Decompress() to get access to
+	  * the payload in FSharedBuffer format.
+	  */
+	 virtual TFuture<FCompressedBuffer> GetCompressedPayload() const = 0;
+
+	/**
+	 * Allows the existing payload to be replaced with a new one.
+	 *
+	 * To pass in a raw pointer use 'FSharedBuffer::...(Data, Size)' to create a valid FSharedBuffer.
+	 * Use 'FSharedBuffer::MakeView' if you want to retain ownership on the data being passed in, and use
+	 * 'FSharedBuffer::TakeOwnership' if you are okay with the bulkdata object taking over ownership of it.
+	 * The bulkdata object must own its internal buffer, so if you pass in a non-owned FSharedBuffer (ie
+	 * by using 'FSharedBuffer::MakeView') then a clone of the data will be created internally and assigned
+	 * to the bulkdata object.
+	 * 
+	 * @param InPayload The payload to update the bulkdata with
+	 * @param CompressionFormat (optional) The compression format to use, NAME_None indicates that the
+	 * payload is already in a compressed format and will not gain from being compressed again. These 
+	 * payloads will never be compressed. NAME_Default will apply which ever compression format that the
+	 * underlying code deems appropriate. Other specific compression formats may be allowed, see the 
+	 * documentation of FCompressedBuffer for details.
+	 */
+	virtual void UpdatePayload(FSharedBuffer InPayload, FName CompressionFormat = NAME_Default) = 0;
+};
+
+} // namespace UE::Virtualization
diff --git a/Engine/Source/Runtime/Core/Public/Virtualization/PayloadId.h b/Engine/Source/Runtime/Core/Public/Virtualization/PayloadId.h
new file mode 100644
--- /dev/null
+++ b/Engine/Source/Runtime/Core/Public/Virtualization/PayloadId.h
@@ -0,0 +1,93 @@
+// Copyright Epic Games, Inc. All Rights Reserved.
+
+#pragma once
+
+#include "Containers/UnrealString.h"
+#include "IO/IoHash.h"
+
+class FArchive;
+class FSharedBuffer;
+
+namespace UE::Virtualization
+{
+	
+/** An identifier for a payload in the virtualization system. */
+class CORE_API FPayloadId
+{
+public:
+	FPayloadId() = default;
+
+	/** Create a FPayloadId from the given FIoHash */
+	explicit FPayloadId(const FIoHash& BlakeHash);
+
+	/** Creates a FPayloadId based on hashing the given FSharedBuffer */
+	explicit FPayloadId(const FSharedBuffer& Payload);
+
+	/** 
+	 * Used when we already have a FGuid identifying the payload and do not want to 
+	 * pay the cost of loading and rehashing the payload each time we load the editor
+	 */
+	explicit FPayloadId(const FGuid& Guid);
+
+	/** Returns true if the id is valid and should reference a payload */
+	bool IsValid() const { return bIsHashSet;  }
+
+	/** Reset the id to the initial, invalid state so that it does not reference a payload */
+	void Reset();
+
+	// TODO:  Although accessors are provided, this is mostly to provide compatibility with 
+	// existing systems. Ideally we should remove these before UE5 is released to prevent their
+	// use by licensees and try to encourage the use of FPayloadId directly.
+
+	/** 
+	 * Returns the id as a FGuid (provides compatibility for systems that expect a FGuid, 
+	 * such as many of the original DDC implementations, over time we should be able to
+	 * phase its use and remove it.
+	 */
+	FGuid ToGuid() const;
+
+	/** Returns a string representation of the id, useful for debug tools */
+	FString ToString() const;
+
+	bool operator == (const FPayloadId& InPayloadId) const;
+	bool operator != (const FPayloadId& InPayloadId) const { return !(*this == InPayloadId); }
+
+	/** Serialization specialization */
+	CORE_API friend FArchive& operator<<(FArchive& Ar, FPayloadId& PayloadId);
+
+	/** Hash specialization */
+	inline friend uint32 GetTypeHash(const FPayloadId& Hash);
+
+	/** String builder specialization */
+	template <typename CharType>
+	friend TStringBuilderBase<CharType>& operator<<(TStringBuilderBase<CharType>& Builder, const FPayloadId& Id);
+private:
+
+	/** The actual identifier */
+	FIoHash Identifier;
+
+	/** 
+	 * Keeps track of if Identifier has been set as checking all elements of a FIoHash isn't 
+	 * exactly fast. Things are further complicated by FIoHash::HashBuffer(Ptr, 0) != FIoHash()
+	 * so we would have two 'hash is not set' states to check for.
+	 */
+	bool bIsHashSet = false;
+};
+
+inline uint32 GetTypeHash(const FPayloadId& Hash)
+{
+	return GetTypeHash(Hash.Identifier);
+}
+
+template <typename CharType>
+TStringBuilderBase<CharType>& operator<<(TStringBuilderBase<CharType>& Builder, const FPayloadId& Id)
+{
+	if (Id.IsValid())
+	{
+		Builder << Id.Identifier;
+	}
+
+	return Builder;
+}
+
+} // namespace UE::Virtualization
