Change the file path to be three directories deep from the root, with each directory being named 00->ff, we then attempt to place each file in a directory based on a hash of the PayloadId, ideally evenly distributing between all available leaf directories.

- Cleaned up some allocations caused by string manipulation.

*Add a place for shared utility code that could have use in multiple places in the virtualization system.
- First utility is a function to generate a file path based off a FPayloadId.
- The path produced is similar to how we organize files in the DDC but instead of 0->9 we use the range 0->ff in hex.
- To generate the path we take the string form of FPayloadId and then use the first 6 characters to create the directories (each directory using 2 characters) and then the remaining 34 characters becomes the filename.

#rnx

diff --git a/Engine/Source/Developer/Virtualization/Private/VirtualizationFileBackend.cpp b/Engine/Source/Developer/Virtualization/Private/VirtualizationFileBackend.cpp
--- a/Engine/Source/Developer/Virtualization/Private/VirtualizationFileBackend.cpp
+++ b/Engine/Source/Developer/Virtualization/Private/VirtualizationFileBackend.cpp
@@ -6,7 +6,9 @@
 #include "HAL/FileManager.h"
 #include "Misc/Parse.h"
 #include "Misc/Paths.h"
+#include "Misc/StringBuilder.h"
 #include "Virtualization/VirtualizationManager.h"
+#include "VirtualizationUtilities.h."
 
 namespace UE::Virtualization
 {
@@ -28,8 +30,7 @@ public:
 	FFileSystemBackend(FStringView ConfigName)
 		: IVirtualizationBackend(EOperations::Both)
 	{
-		Name = TEXT("FFileSystemBackend - ");
-		Name.Append(ConfigName);
+		Name = WriteToString<256>(TEXT("FFileSystemBackend - "), ConfigName).ToString();
 	}
 
 	virtual ~FFileSystemBackend() = default;
@@ -69,14 +70,15 @@ private:
 			return true;
 		}
 
-		const FString FilePath = CreateFilePath(Id);
+		TStringBuilder<512> FilePath;
+		CreateFilePath(Id, FilePath);
 
 		{
 			// TODO: Should we write to a temp file and then move it once it has written?
-			TUniquePtr<FArchive> FileAr(IFileManager::Get().CreateFileWriter(*FilePath));
+			TUniquePtr<FArchive> FileAr(IFileManager::Get().CreateFileWriter(FilePath.ToString()));
 			if (FileAr == nullptr)
 			{
-				UE_LOG(LogVirtualization, Error, TEXT("[%s] Failed to push payload '%s' to '%s'"), *GetDebugString(), *Id.ToString(), *FilePath);
+				UE_LOG(LogVirtualization, Error, TEXT("[%s] Failed to push payload '%s' to '%s'"), *GetDebugString(), *Id.ToString(), FilePath.ToString());
 				return false;
 			}
 
@@ -94,16 +96,17 @@ private:
 	{
 		TRACE_CPUPROFILER_EVENT_SCOPE(FFileSystemBackend::PullData);
 
-		const FString FilePath = CreateFilePath(Id);
+		TStringBuilder<512> FilePath;
+		CreateFilePath(Id, FilePath);
 
 		// TODO: Should we allow the error severity to be configured via ini or just not report this case at all?
-		if (!IFileManager::Get().FileExists(*FilePath))
+		if (!IFileManager::Get().FileExists(FilePath.ToString()))
 		{
 			UE_LOG(LogVirtualization, Verbose, TEXT("[%s] Does not contain the payload '%s'"), *GetDebugString(), *Id.ToString());
 			return FCompressedBuffer();
 		}
 
-		TUniquePtr<FArchive> FileAr(IFileManager::Get().CreateFileReader(*FilePath));
+		TUniquePtr<FArchive> FileAr(IFileManager::Get().CreateFileReader(FilePath.ToString()));
 		if (FileAr == nullptr)
 		{
 			const uint32 SystemError = FPlatformMisc::GetLastError();
@@ -113,11 +116,11 @@ private:
 			{
 				TCHAR SystemErrorMsg[2048] = { 0 };
 				FPlatformMisc::GetSystemErrorMessage(SystemErrorMsg, sizeof(SystemErrorMsg), SystemError);
-				UE_LOG(LogVirtualization, Error, TEXT("[%s] Failed to load payload '%s' file '%s' due to system error: '%s' (%d))"), *GetDebugString(), *Id.ToString(), *FilePath, SystemErrorMsg, SystemError);
+				UE_LOG(LogVirtualization, Error, TEXT("[%s] Failed to load payload '%s' file '%s' due to system error: '%s' (%d))"), *GetDebugString(), *Id.ToString(), FilePath.ToString(), SystemErrorMsg, SystemError);
 			}
 			else
 			{
-				UE_LOG(LogVirtualization, Error, TEXT("[%s] Failed to load payload '%s' from '%s' (reason unknown)"), *GetDebugString(), *Id.ToString(), *FilePath);
+				UE_LOG(LogVirtualization, Error, TEXT("[%s] Failed to load payload '%s' from '%s' (reason unknown)"), *GetDebugString(), *Id.ToString(), FilePath.ToString());
 			}
 		
 			return FCompressedBuffer();
@@ -130,8 +133,10 @@ private:
 	{
 		TRACE_CPUPROFILER_EVENT_SCOPE(FFileSystemBackend::DoesExist);
 
-		const FString FilePath = CreateFilePath(Id);
-		return IFileManager::Get().FileExists(*FilePath);
+		TStringBuilder<512> FilePath;
+		CreateFilePath(Id, FilePath);
+
+		return IFileManager::Get().FileExists(FilePath.ToString());
 	}
 
 	virtual FString GetDebugString() const override
@@ -139,12 +144,16 @@ private:
 		return Name;
 	}
 
-	FString CreateFilePath(const FPayloadId& PayloadId)
+	void CreateFilePath(const FPayloadId& PayloadId, FStringBuilderBase& OutPath)
 	{
-		return FString::Printf(TEXT("%s/%s.payload"), *RootDirectory, *PayloadId.ToString());
+		TStringBuilder<52> PayloadPath;
+		Utils::PayloadIdToPath(PayloadId, PayloadPath);
+
+		OutPath << RootDirectory << TEXT("/") << PayloadPath;
 	}
 
 private:
+
 	FString Name;
 	FString RootDirectory;
 };
diff --git a/Engine/Source/Developer/Virtualization/Private/VirtualizationUtilities.cpp b/Engine/Source/Developer/Virtualization/Private/VirtualizationUtilities.cpp
new file mode 100644
--- /dev/null
+++ b/Engine/Source/Developer/Virtualization/Private/VirtualizationUtilities.cpp
@@ -0,0 +1,34 @@
+// Copyright Epic Games, Inc. All Rights Reserved.
+
+#include "VirtualizationUtilities.h"
+
+#include "Misc/StringBuilder.h"
+#include "Virtualization/PayloadId.h"
+
+namespace UE::Virtualization::Utils
+{
+
+void PayloadIdToPath(const FPayloadId& Id, FStringBuilderBase& OutPath)
+{
+	OutPath.Reset();
+	OutPath << Id;
+
+	TStringBuilder<10> Directory;
+	Directory << OutPath.ToView().Left(2) << TEXT("/");
+	Directory << OutPath.ToView().Mid(2, 2) << TEXT("/");
+	Directory << OutPath.ToView().Mid(4, 2) << TEXT("/");
+
+	OutPath.ReplaceAt(0, 6, Directory);
+
+	OutPath << TEXT(".payload");
+}
+
+FString PayloadIdToPath(const FPayloadId& Id)
+{
+	TStringBuilder<52> Path;
+	PayloadIdToPath(Id, Path);
+
+	return FString(Path);
+}
+
+} // namespace UE::Virtualization::Utils
diff --git a/Engine/Source/Developer/Virtualization/Private/VirtualizationUtilities.h b/Engine/Source/Developer/Virtualization/Private/VirtualizationUtilities.h
new file mode 100644
--- /dev/null
+++ b/Engine/Source/Developer/Virtualization/Private/VirtualizationUtilities.h
@@ -0,0 +1,44 @@
+// Copyright Epic Games, Inc. All Rights Reserved.
+
+#pragma once
+
+#include "Containers/StringFwd.h"
+#include "Containers/UnrealString.h"
+
+namespace UE::Virtualization
+{
+	class FPayloadId;
+} // namespace UE::Virtualization
+
+namespace UE::Virtualization::Utils
+{
+
+/** 
+ * Converts a FPayloadId into a file path.
+ *
+ * This utility will take a FPayloadId and return a file path that is
+ * 3 directories deep. The first six characters of the id will be used to
+ * create the directory names, with each directory using two characters.
+ * The remaining thirty four characters will be used as the file name.
+ * Lastly the extension '.payload' will be applied to complete the path/
+ * Example: FPayloadId 0139d6d5d477e32dfd2abd3c5bc8ea8507e8eef8 becomes
+ *			01/39/d6/d5d477e32dfd2abd3c5bc8ea8507e8eef8.payload'
+ * 
+ * @param	Id The payload identifier used to create the file path.
+ * @param	OutPath Will be reset and then assigned the resulting file path.
+ *			The string builder must have a capacity of at least 52 characters 
+			to avoid reallocation.
+ */
+void PayloadIdToPath(const FPayloadId& Id, FStringBuilderBase& OutPath);
+
+/** 
+ * Converts a FPayloadId into a file path.
+ * 
+ * See above for further details
+ * 
+ * @param	Id The payload identifier used to create the file path.
+ * @return	The resulting file path.
+ */
+FString PayloadIdToPath(const FPayloadId& Id);
+
+} // namespace UE::Virtualization::Utils
