CIS: Fix typo in include

#rnx

diff --git a/Engine/Source/Developer/Virtualization/Private/VirtualizationFileBackend.cpp b/Engine/Source/Developer/Virtualization/Private/VirtualizationFileBackend.cpp
--- a/Engine/Source/Developer/Virtualization/Private/VirtualizationFileBackend.cpp
+++ b/Engine/Source/Developer/Virtualization/Private/VirtualizationFileBackend.cpp
@@ -8,7 +8,7 @@
 #include "Misc/Paths.h"
 #include "Misc/StringBuilder.h"
 #include "Virtualization/VirtualizationManager.h"
-#include "VirtualizationUtilities.h."
+#include "VirtualizationUtilities.h"
 
 namespace UE::Virtualization
 {
