Add a Mirage backend that can access data in source control.

#rnx

diff --git a/Engine/Source/Developer/Virtualization/Private/VirtualizationSourceControlBackend.cpp b/Engine/Source/Developer/Virtualization/Private/VirtualizationSourceControlBackend.cpp
new file mode 100644
--- /dev/null
+++ b/Engine/Source/Developer/Virtualization/Private/VirtualizationSourceControlBackend.cpp
@@ -0,0 +1,169 @@
+// Copyright Epic Games, Inc. All Rights Reserved.
+
+#include "Virtualization/IVirtualizationBackend.h"
+
+#include "ISourceControlModule.h"
+#include "ISourceControlProvider.h"
+#include "Misc/Parse.h"
+#include "SourceControlOperations.h"
+#include "Virtualization/PayloadId.h"
+#include "VirtualizationSourceControlUtilities.h"
+#include "VirtualizationUtilities.h"
+
+// When the SourceControl module (or at least the perforce source control module) is thread safe we
+// can enable this and stop using the hacky work around 'TryToDownloadFileFromBackgroundThread'
+#define IS_SOURCE_CONTROL_THREAD_SAFE 0
+
+namespace UE::Virtualization
+{
+
+/**
+ * This backend can be used to access payloads stored in source control.
+ * The backend doesn't 'check out' a payload file but instead will just download the payload as
+ * a binary blob.
+ * It is assumed that the files are stored with the same path convention as the file system 
+ * backend, found in Utils::PayloadIdToPath.
+ * 
+ * Ini file setup:
+ * 'Name'=(Type=SourceControl, DepotRoot="//XXX/")
+ * Where 'Name' is the backend name in the hierarchy and 'XXX' is the path in the source control
+ * depot where the payload files are being stored.
+ */
+class FSourceControlBackend : public IVirtualizationBackend
+{
+public:
+	FSourceControlBackend(FStringView ConfigName)
+		: IVirtualizationBackend(EOperations::Pull)
+	{
+	}
+
+	virtual bool Initialize(const FString& ConfigEntry) override
+	{
+		TRACE_CPUPROFILER_EVENT_SCOPE(FSourceControlBackend::Initialize);
+
+		// We require that a valid depot root has been provided
+		if (!FParse::Value(*ConfigEntry, TEXT("DepotRoot="), DepotRoot))
+		{
+			UE_LOG(LogVirtualization, Error, TEXT("'DepotRoot=' not found in the config file"));
+			return false;
+		}
+
+		const ISourceControlModule& SSCModule = ISourceControlModule::Get();
+
+		// We require source control the be enabled 
+		if (!SSCModule.IsEnabled())
+		{
+			UE_LOG(LogVirtualization, Error, TEXT("Attempting to initialize FSourceControlBackend but source control is disabled!"));
+			return false;
+		}
+
+		ISourceControlProvider& SCCProvider = SSCModule.GetProvider();
+
+		// We require perforce as the source control provider as it is currently the only one that has the virtualization functionality implemented
+		const FName SourceControlName = SCCProvider.GetName();
+		if (SourceControlName != FName("Perforce"))
+		{
+			UE_LOG(LogVirtualization, Error, TEXT("Attempting to initialize FSourceControlBackend but source control is '%s' and only Perforce is currently supported!"), *SourceControlName.ToString());
+			return false;
+		}
+
+		// When a source control depot is set up a file named 'payload_metainfo.txt' should be submitted to it's root.
+		// This allows us to check for the existence of the file to confirm that the depot root is indeed valid.
+		const FString PayloadMetaInfoPath = FString::Printf(TEXT("%spayload_metainfo.txt"), *DepotRoot);
+
+#if IS_SOURCE_CONTROL_THREAD_SAFE
+		TSharedRef<FDownloadFile, ESPMode::ThreadSafe> DownloadCommand = ISourceControlOperation::Create<FDownloadFile>();
+		if (SCCProvider.Execute(DownloadCommand, PayloadMetaInfoPath, EConcurrency::Synchronous) != ECommandResult::Succeeded)
+		{
+			UE_LOG(LogVirtualization, Error, TEXT("Failed to find 'payload_metainfo.txt' in the depot '%s', is your config set up correctly?"), *DepotRoot);
+			return false;
+		}	
+#else
+		TSharedRef<FDownloadFile, ESPMode::ThreadSafe> DownloadCommand = ISourceControlOperation::Create<FDownloadFile>();
+		if (!SCCProvider.TryToDownloadFileFromBackgroundThread(DownloadCommand, PayloadMetaInfoPath))
+		{
+			UE_LOG(LogVirtualization, Error, TEXT("Failed to find 'payload_metainfo.txt' in the depot '%s', is your config set up correctly?"), *DepotRoot);
+			return false;
+		}		
+#endif //IS_SOURCE_CONTROL_THREAD_SAFE
+
+		FSharedBuffer MetaInfoBuffer = DownloadCommand->GetFileData(PayloadMetaInfoPath);
+		if (MetaInfoBuffer.IsNull())
+		{
+			UE_LOG(LogVirtualization, Error, TEXT("Failed to find 'payload_metainfo.txt' in the depot '%s', is your config set up correctly?"), *DepotRoot);
+			return false;
+		}
+
+		// Currently we do not do anything with the payload meta info, in the future we could structure
+		// it's format to include more information that might be worth logging or something. 
+		// But for now being able to pull the payload meta info path at least shows that we can use the
+		// depot.
+
+		return true;
+	}
+
+	virtual bool PushData(const FPayloadId& Id, const FCompressedBuffer& Payload) override
+	{
+		// This backend will not actually push data to source control, that will be done by
+		// a separate submission tool. Since files submitted to source control are there forever 
+		// we don't want the user to be uploaded new entries each time that they save the asset
+		// but instead upload once when the package is committed to source control.
+		// TODO: As we put the pieces together it is likely that we will change the backend API
+		// so that we don't need dummy implementations of PushData for backends that don't need it.
+		// Especially considering that it is most likely that no backend will push.
+
+		checkNoEntry(); 
+		return false;
+	}
+
+	virtual FCompressedBuffer PullData(const FPayloadId& Id) override
+	{
+		TStringBuilder<512> DepotPath;
+		CreateDepotPath(Id, DepotPath);
+
+		ISourceControlProvider& SCCProvider = ISourceControlModule::Get().GetProvider();
+
+#if IS_SOURCE_CONTROL_THREAD_SAFE
+		TSharedRef<FDownloadFile, ESPMode::ThreadSafe> DownloadCommand = ISourceControlOperation::Create<FDownloadFile>();
+		if (SCCProvider.Execute(DownloadCommand, DepotPath.ToString(), EConcurrency::Synchronous) != ECommandResult::Succeeded)
+		{
+			return FCompressedBuffer();
+		}
+#else
+		TSharedRef<FDownloadFile> DownloadCommand = ISourceControlOperation::Create<FDownloadFile>();
+		if (!SCCProvider.TryToDownloadFileFromBackgroundThread(DownloadCommand, DepotPath.ToString()))
+		{
+			return FCompressedBuffer();
+		}
+#endif
+
+		// The payload was created by FCompressedBuffer::Compress so we can return it 
+		// as a FCompressedBuffer.
+		FSharedBuffer Buffer = DownloadCommand->GetFileData(DepotPath);
+		return FCompressedBuffer::FromCompressed(Buffer);
+	}
+
+	virtual FString GetDebugString() const override
+	{
+		return FString(TEXT("SourceControl"));
+	}
+
+	void CreateDepotPath(const FPayloadId& PayloadId, FStringBuilderBase& OutPath)
+	{
+		TStringBuilder<52> PayloadPath;
+		Utils::PayloadIdToPath(PayloadId, PayloadPath);
+
+		OutPath << DepotRoot << PayloadPath;
+	}
+
+private:
+
+	/** The root where the virtualized payloads are stored in source control */
+	FString DepotRoot;
+};
+
+UE_REGISTER_VIRTUALIZATION_BACKEND_FACTORY(FSourceControlBackend, SourceControl);
+
+} // namespace UE::Virtualization
+
+PRAGMA_ENABLE_OPTIMIZATION
