The mirage Jupiter backend by default should not attempt to upload a payload that already exists in the service.

#rnx

diff --git a/Engine/Source/Developer/Virtualization/Private/VirtualizationJupiterBackend.cpp b/Engine/Source/Developer/Virtualization/Private/VirtualizationJupiterBackend.cpp
--- a/Engine/Source/Developer/Virtualization/Private/VirtualizationJupiterBackend.cpp
+++ b/Engine/Source/Developer/Virtualization/Private/VirtualizationJupiterBackend.cpp
@@ -65,7 +65,7 @@
 
 /** When enabled we will only attempt to upload a payload if Jupiter claims to not already have it.
 	Disabling this is only intended for debug purposes. */
-#define UE_CHECK_FOR_EXISTING_PAYLOADS 0
+#define UE_CHECK_FOR_EXISTING_PAYLOADS 1
 
 /** When enabled we will only attempt to upload a chunk if Jupiter claims to not already have it. 
 	In practice this doesn't really speed up the workflow hence being disabled. */
