CIS Fix: Remove accidental pragma left at the bottom of a file.

#rnx

diff --git a/Engine/Source/Developer/Virtualization/Private/VirtualizationSourceControlBackend.cpp b/Engine/Source/Developer/Virtualization/Private/VirtualizationSourceControlBackend.cpp
--- a/Engine/Source/Developer/Virtualization/Private/VirtualizationSourceControlBackend.cpp
+++ b/Engine/Source/Developer/Virtualization/Private/VirtualizationSourceControlBackend.cpp
@@ -165,5 +165,3 @@ private:
 UE_REGISTER_VIRTUALIZATION_BACKEND_FACTORY(FSourceControlBackend, SourceControl);
 
 } // namespace UE::Virtualization
-
-PRAGMA_ENABLE_OPTIMIZATION
