IVirtualizationBackend::PushData now returns an enum describing what happened to the operation rather than a straight true/false.

#rnx

diff --git a/Engine/Source/Developer/Virtualization/Private/VirtualizationFileBackend.cpp b/Engine/Source/Developer/Virtualization/Private/VirtualizationFileBackend.cpp
--- a/Engine/Source/Developer/Virtualization/Private/VirtualizationFileBackend.cpp
+++ b/Engine/Source/Developer/Virtualization/Private/VirtualizationFileBackend.cpp
@@ -60,14 +60,14 @@ private:
 		return true;
 	}
 
-	virtual bool PushData(const FPayloadId& Id, const FCompressedBuffer& Payload) override
+	virtual EPushResult PushData(const FPayloadId& Id, const FCompressedBuffer& Payload) override
 	{
 		TRACE_CPUPROFILER_EVENT_SCOPE(FFileSystemBackend::PushData);
 
 		if (DoesExist(Id))
 		{
 			UE_LOG(LogVirtualization, Verbose, TEXT("[%s] Already has a copy of the payload '%s'."), *GetDebugString(), *Id.ToString());
-			return true;
+			return EPushResult::PayloadAlreadyExisted;
 		}
 
 		TStringBuilder<512> FilePath;
@@ -79,7 +79,7 @@ private:
 			if (FileAr == nullptr)
 			{
 				UE_LOG(LogVirtualization, Error, TEXT("[%s] Failed to push payload '%s' to '%s'"), *GetDebugString(), *Id.ToString(), FilePath.ToString());
-				return false;
+				return EPushResult::Failed;
 			}
 
 			for (const FSharedBuffer& Buffer : Payload.GetCompressed().GetSegments())
@@ -89,7 +89,7 @@ private:
 			}
 		}
 
-		return true;
+		return EPushResult::Success;
 	}
 
 	virtual FCompressedBuffer PullData(const FPayloadId& Id) override
diff --git a/Engine/Source/Developer/Virtualization/Private/VirtualizationJupiterBackend.cpp b/Engine/Source/Developer/Virtualization/Private/VirtualizationJupiterBackend.cpp
--- a/Engine/Source/Developer/Virtualization/Private/VirtualizationJupiterBackend.cpp
+++ b/Engine/Source/Developer/Virtualization/Private/VirtualizationJupiterBackend.cpp
@@ -1397,7 +1397,7 @@ private:
 		return true; 
 	}
 
-	virtual bool PushData(const FPayloadId& Id, const FCompressedBuffer& CompressedPayload) override
+	virtual EPushResult PushData(const FPayloadId& Id, const FCompressedBuffer& CompressedPayload) override
 	{
 		using namespace Utility;
 
@@ -1407,7 +1407,7 @@ private:
 		if (DoesPayloadExist(Id))
 		{
 			UE_LOG(LogVirtualization, Verbose, TEXT("Jupiter already has a copy of the payload '%s'"), *Id.ToString());
-			return true;
+			return EPushResult::PayloadAlreadyExisted;
 		}
 #endif // UE_CHECK_FOR_EXISTING_PAYLOADS
 
@@ -1423,7 +1423,7 @@ private:
 		if (NumChunks > MAX_int32)
 		{
 			UE_LOG(LogVirtualization, Error, TEXT("Too many chunks (%d) are required for the payload '%s', try increasing the ChunkSize"), NumChunks, *Id.ToString());
-			return false;
+			return EPushResult::Failed;
 		}
 
 		PUTRequest.ChunkHashes.SetNum((int32)NumChunks);
@@ -1473,7 +1473,7 @@ private:
 		if (NumFailedChunks > 0)
 		{
 			UE_LOG(LogVirtualization, Error, TEXT("Failed to upload %d chunks for the payload '%s'."), NumFailedChunks.load(), *Id.ToString());
-			return false;
+			return EPushResult::Failed;
 		}
 
 		UE_LOG(LogVirtualization, Verbose, TEXT("Successfully uploaded all chunks for the payload '%s'"), *Id.ToString());	
@@ -1500,19 +1500,19 @@ private:
 				if (ResponseCode == 200)
 				{
 					UE_LOG(LogVirtualization, Verbose, TEXT("Successfully uploaded the description for the payload '%s'"), *Id.ToString());	
-					return true;
+					return EPushResult::Success;
 				}
 
 				if (!ShouldRetryOnError(ResponseCode))
 				{
 					UE_LOG(LogVirtualization, Error, TEXT("Failed with error code '%d' to upload header infomation about payload '%s'"), ResponseCode, *Id.ToString());
-					return false;
+					return EPushResult::Failed;
 				}
 			}
 		}
 
 		UE_LOG(LogVirtualization, Error, TEXT("Failed  '%d' attempts to upload header infomation about payload '%s'"), UE_MIRAGE_MAX_ATTEMPTS, *Id.ToString());
-		return false;
+		return EPushResult::Failed;
 	}
 
 	virtual FCompressedBuffer PullData(const FPayloadId& Id) override
diff --git a/Engine/Source/Developer/Virtualization/Private/VirtualizationSourceControlBackend.cpp b/Engine/Source/Developer/Virtualization/Private/VirtualizationSourceControlBackend.cpp
--- a/Engine/Source/Developer/Virtualization/Private/VirtualizationSourceControlBackend.cpp
+++ b/Engine/Source/Developer/Virtualization/Private/VirtualizationSourceControlBackend.cpp
@@ -102,7 +102,7 @@ public:
 		return true;
 	}
 
-	virtual bool PushData(const FPayloadId& Id, const FCompressedBuffer& Payload) override
+	virtual EPushResult PushData(const FPayloadId& Id, const FCompressedBuffer& Payload) override
 	{
 		// This backend will not actually push data to source control, that will be done by
 		// a separate submission tool. Since files submitted to source control are there forever 
@@ -113,7 +113,7 @@ public:
 		// Especially considering that it is most likely that no backend will push.
 
 		checkNoEntry(); 
-		return false;
+		return EPushResult::Failed;
 	}
 
 	virtual FCompressedBuffer PullData(const FPayloadId& Id) override
diff --git a/Engine/Source/Runtime/CoreUObject/Private/Virtualization/VirtualizationManager.cpp b/Engine/Source/Runtime/CoreUObject/Private/Virtualization/VirtualizationManager.cpp
--- a/Engine/Source/Runtime/CoreUObject/Private/Virtualization/VirtualizationManager.cpp
+++ b/Engine/Source/Runtime/CoreUObject/Private/Virtualization/VirtualizationManager.cpp
@@ -185,16 +185,19 @@ bool FVirtualizationManager::PushData(const FPayloadId& Id, const FCompressedBuf
 	bool bWasPayloadPushed = false;
 	for (IVirtualizationBackend* Backend : PushEnabledBackendsArray)
 	{
-		const bool bDidPushPayload = Backend->PushData(Id, Payload);
-		
-		UE_CLOG(bDidPushPayload, LogVirtualization, Verbose, TEXT("[%s] Pushed the payload '%s'"), *Backend->GetDebugString(), *Id.ToString());
-		UE_CLOG(!bDidPushPayload, LogVirtualization, Error, TEXT("[%s] Failed to push the payload '%s'"), *Backend->GetDebugString(), *Id.ToString());
+		const EPushResult Result = Backend->PushData(Id, Payload);
+
+		UE_CLOG(Result != EPushResult::Failed, LogVirtualization, Verbose, TEXT("[%s] Pushed the payload '%s'"), *Backend->GetDebugString(), *Id.ToString());
+		UE_CLOG(Result == EPushResult::Failed, LogVirtualization, Error, TEXT("[%s] Failed to push the payload '%s'"), *Backend->GetDebugString(), *Id.ToString());
 
-		bWasPayloadPushed |= bDidPushPayload;
+		if (Result != EPushResult::Failed)
+		{
+			bWasPayloadPushed = true;
+		}
 		
 		// Debugging operation where we immediately try to pull the payload after each push (when possible) and assert 
 		// that the pulled payload is the same as the original
-		if (bValidateAfterPushOperation && bDidPushPayload && Backend->SupportsPullOperations())
+		if (bValidateAfterPushOperation && Result != EPushResult::Failed && Backend->SupportsPullOperations())
 		{
 			FCompressedBuffer PulledPayload = Backend->PullData(Id);
 			checkf(Payload.GetRawHash() == PulledPayload.GetRawHash(), TEXT("[%s] Failed to pull payload '%s' after it was pushed to backend"), 
diff --git a/Engine/Source/Runtime/CoreUObject/Public/Virtualization/IVirtualizationBackend.h b/Engine/Source/Runtime/CoreUObject/Public/Virtualization/IVirtualizationBackend.h
--- a/Engine/Source/Runtime/CoreUObject/Public/Virtualization/IVirtualizationBackend.h
+++ b/Engine/Source/Runtime/CoreUObject/Public/Virtualization/IVirtualizationBackend.h
@@ -12,6 +12,17 @@ namespace UE::Virtualization
 
 class FPayloadId;
 
+/** Describes the result of a IVirtualizationBackend::Push operation */
+enum class EPushResult
+{
+	/** The push failed, the backend should print an error message to 'LogVirtualization'.*/
+	Failed = 0,
+	/** The payload already exists in the backend and does not need to be pushed. */
+	PayloadAlreadyExisted,
+	/** The payload was successfully pushed to the backend. */
+	Success
+};
+
 /**
  * The interface to derive from to create a new backend implementation.
  * 
@@ -68,7 +79,7 @@ public:
 	 * @param Payload	A potentially compressed buffer representing the payload
 	 * @return			The result of the push operation
 	 */
-	virtual bool PushData(const FPayloadId& Id, const FCompressedBuffer& Payload) = 0;
+	virtual EPushResult PushData(const FPayloadId& Id, const FCompressedBuffer& Payload) = 0;
 
 	/** 
 	 * The backend will attempt to retrieve the given payload by what ever method the backend uses.
