Add a cpu trace when pulling payloads for the source control backend.

#rnx

diff --git a/Engine/Source/Developer/Virtualization/Private/VirtualizationSourceControlBackend.cpp b/Engine/Source/Developer/Virtualization/Private/VirtualizationSourceControlBackend.cpp
--- a/Engine/Source/Developer/Virtualization/Private/VirtualizationSourceControlBackend.cpp
+++ b/Engine/Source/Developer/Virtualization/Private/VirtualizationSourceControlBackend.cpp
@@ -123,6 +123,8 @@ public:
 
 	virtual FCompressedBuffer PullData(const FPayloadId& Id) override
 	{
+		TRACE_CPUPROFILER_EVENT_SCOPE(FSourceControlBackend::PullData);
+
 		TStringBuilder<512> DepotPath;
 		CreateDepotPath(Id, DepotPath);
 
