Improve error handling and logging in the file and jupiter Mirage backends.

#rnx

* VirtualizationJupiterBackend
- When pulling a payload we should assume that a 400 error response when trying to GET the payload header means that the payload is not in Jupiter.
-- Not being able to find the payload should not log an error, instead we can make a note of it in the verbose log (similar to the file system backend)

* VirtualizationFileBackend
- Moved the formatting of system errors to it's own function.
- Log the system error when failing to write a payload during a push as well as a pull.
- We now check that the FileArchive wrote correctly to disk and delete the output file and fail the push if it did not.
-- A future piece of work will change the logic to write to a tmp file at the root of the file store and them move the file to the final location to cut down on the potential of leaving corrupted files around (similar to the process when we save packages)

* Perforce
- The FDownloadFile command now takes an optional parameter EVerbosity that can allow the caller to choose the level of logging output that the command will generate.
- The source control backend for Mirage now opts to supress the logging of the full perforce command when we are pulling payloads as we can generate many hundreds or thousands of requests and the info is not useful to users.
-- We continue to log the command when validating the depot as this is the most likely command to fail so having the info in the log may prove useful.

diff --git a/Engine/Plugins/Developer/PerforceSourceControl/Source/PerforceSourceControl/Private/PerforceSourceControlOperations.cpp b/Engine/Plugins/Developer/PerforceSourceControl/Source/PerforceSourceControl/Private/PerforceSourceControlOperations.cpp
--- a/Engine/Plugins/Developer/PerforceSourceControl/Source/PerforceSourceControl/Private/PerforceSourceControlOperations.cpp
+++ b/Engine/Plugins/Developer/PerforceSourceControl/Source/PerforceSourceControl/Private/PerforceSourceControlOperations.cpp
@@ -2820,8 +2820,11 @@ bool FPerforceDownloadFileWorker::Execute(FPerforceSourceControlCommand& InComma
 				Parameters.Add(TargetFilePath);
 
 				TOptional<FSharedBuffer> FileData = FSharedBuffer();
-				InCommand.bCommandSuccessful = Connection.RunCommand(TEXT("print"), Parameters, Records, FileData, InCommand.ResultInfo.ErrorMessages, FOnIsCancelled::CreateRaw(&InCommand, &FPerforceSourceControlCommand::IsCanceled), InCommand.bConnectionDropped);
 
+				InCommand.bCommandSuccessful = Connection.RunCommand(	TEXT("print"), Parameters, Records, FileData, InCommand.ResultInfo.ErrorMessages, 
+																		FOnIsCancelled::CreateRaw(&InCommand, &FPerforceSourceControlCommand::IsCanceled), 
+																		InCommand.bConnectionDropped, Operation->ShouldLogToStdOutput(), true);
+				
 				if (InCommand.bCommandSuccessful)
 				{
 					Operation->AddFileData(TargetFilePath, FileData.GetValue());
@@ -2848,7 +2851,11 @@ bool FPerforceDownloadFileWorker::Execute(FPerforceSourceControlCommand& InComma
 				Parameters.Add(TEXT("-q")); // Do not print the header, we only want the actual file contents
 				Parameters.Add(TargetFilePath);
 
-				InCommand.bCommandSuccessful = Connection.RunCommand(TEXT("print"), Parameters, Records, InCommand.ResultInfo.ErrorMessages, FOnIsCancelled::CreateRaw(&InCommand, &FPerforceSourceControlCommand::IsCanceled), InCommand.bConnectionDropped);
+				TOptional<FSharedBuffer> NullBuffer;
+
+				InCommand.bCommandSuccessful = Connection.RunCommand(	TEXT("print"), Parameters, Records, NullBuffer, InCommand.ResultInfo.ErrorMessages, 
+																		FOnIsCancelled::CreateRaw(&InCommand, &FPerforceSourceControlCommand::IsCanceled), 
+																		InCommand.bConnectionDropped, Operation->ShouldLogToStdOutput(), true);
 
 				if (InCommand.bCommandSuccessful)
 				{
diff --git a/Engine/Source/Developer/SourceControl/Private/SourceControlOperations.cpp b/Engine/Source/Developer/SourceControl/Private/SourceControlOperations.cpp
--- a/Engine/Source/Developer/SourceControl/Private/SourceControlOperations.cpp
+++ b/Engine/Source/Developer/SourceControl/Private/SourceControlOperations.cpp
@@ -6,8 +6,9 @@
 #include "ISourceControlModule.h"
 #include "Misc/Paths.h"
 
-FDownloadFile::FDownloadFile(FStringView InTargetDirectory)
-	: TargetDirectory(InTargetDirectory)
+FDownloadFile::FDownloadFile(FStringView InTargetDirectory, EVerbosity InVerbosity)
+	: Verbosity(InVerbosity)
+	, TargetDirectory(InTargetDirectory)
 {
 	FPaths::NormalizeDirectoryName(TargetDirectory);
 
diff --git a/Engine/Source/Developer/SourceControl/Public/SourceControlOperations.h b/Engine/Source/Developer/SourceControl/Public/SourceControlOperations.h
--- a/Engine/Source/Developer/SourceControl/Public/SourceControlOperations.h
+++ b/Engine/Source/Developer/SourceControl/Public/SourceControlOperations.h
@@ -671,19 +671,32 @@ public:
 };
 
 /**
- * Operation used to download a file from the source control server
+ * Operation used to download a file from the source control server directly rather
+ * than sync it. This should not change the state of anything locally on the client.
  */
 class FDownloadFile : public FSourceControlOperationBase
 {
 public:
+	enum class EVerbosity
+	{
+		/** No logging when the command is run. */
+		None,
+		/** Log the full cmdline when the command is run. */
+		Full
+	};
+
 	/** 
 	 * This constructor will download the files and keep them in memory, 
 	 * which can then be accessed by calling XXX.
 	 */
 	FDownloadFile() = default;
+
+	FDownloadFile( EVerbosity InVerbosity)
+		: Verbosity(InVerbosity)
+	{}
 	
 	/** This constructor will download the files to the given directory */
-	SOURCECONTROL_API FDownloadFile(FStringView InTargetDirectory);
+	SOURCECONTROL_API FDownloadFile(FStringView InTargetDirectory, EVerbosity InVerbosity);
 
 	// ISourceControlOperation interface
 	virtual FName GetName() const override
@@ -724,7 +737,13 @@ public:
 			return FSharedBuffer();
 		}
 	}
+
+	bool ShouldLogToStdOutput() const
+	{
+		return Verbosity == EVerbosity::Full;
+	}
 private:
+	EVerbosity Verbosity = EVerbosity::Full;
 
 	FString TargetDirectory;
 	TMap<FString, FSharedBuffer> FileDataMap;
diff --git a/Engine/Source/Developer/Virtualization/Private/VirtualizationFileBackend.cpp b/Engine/Source/Developer/Virtualization/Private/VirtualizationFileBackend.cpp
--- a/Engine/Source/Developer/Virtualization/Private/VirtualizationFileBackend.cpp
+++ b/Engine/Source/Developer/Virtualization/Private/VirtualizationFileBackend.cpp
@@ -12,6 +12,32 @@
 
 namespace UE::Virtualization
 {
+/**
+ * Fill in the given string builder with the human readable message of the current system
+ * code, followed by the code value itself. 
+ * In the system value is currently 0, then we assume that it was cleared before this was 
+ * able to be called and write that the error is unknown instead of assuming that the 
+ * operation was a success.
+ */
+void GetFormattedSystemError(FStringBuilderBase& SystemErrorMessage)
+{
+	SystemErrorMessage.Reset();
+
+	const uint32 SystemError = FPlatformMisc::GetLastError();
+	// If we have a system error we can give a more informative error message but don't output it if the error is zero as 
+	// this can lead to very confusing error messages.
+	if (SystemError != 0)
+	{
+		TCHAR SystemErrorMsg[MAX_SPRINTF] = { 0 };
+		FPlatformMisc::GetSystemErrorMessage(SystemErrorMsg, sizeof(SystemErrorMsg), SystemError);
+
+		SystemErrorMessage.Appendf(TEXT("'%s' (%d)"), SystemErrorMsg, SystemError);
+	}
+	else
+	{
+		SystemErrorMessage << TEXT("'unknown reason' (0)");
+	}
+}
 
 /**
  * A basic backend based on the file system. This can be used to access/store virtualization
@@ -73,23 +99,35 @@ private:
 		TStringBuilder<512> FilePath;
 		CreateFilePath(Id, FilePath);
 
+		// TODO: Should we write to a temp file and then move it once it has written?
+		TUniquePtr<FArchive> FileAr(IFileManager::Get().CreateFileWriter(FilePath.ToString()));
+
+		if (FileAr == nullptr)
+		{
+			TStringBuilder<MAX_SPRINTF> SystemErrorMsg;
+			GetFormattedSystemError(SystemErrorMsg);
+			UE_LOG(LogVirtualization, Error, TEXT("[%s] Failed to write payload '%s' to '%s' due to system error: %s"), *GetDebugString(), *Id.ToString(), FilePath.ToString(), SystemErrorMsg.ToString());
+
+			return EPushResult::Failed;
+		}
+
+		for (const FSharedBuffer& Buffer : Payload.GetCompressed().GetSegments())
 		{
-			// TODO: Should we write to a temp file and then move it once it has written?
-			TUniquePtr<FArchive> FileAr(IFileManager::Get().CreateFileWriter(FilePath.ToString()));
-			if (FileAr == nullptr)
-			{
-				UE_LOG(LogVirtualization, Error, TEXT("[%s] Failed to push payload '%s' to '%s'"), *GetDebugString(), *Id.ToString(), FilePath.ToString());
-				return EPushResult::Failed;
-			}
-
-			for (const FSharedBuffer& Buffer : Payload.GetCompressed().GetSegments())
-			{
-				// Const cast because FArchive requires a non-const pointer!
-				FileAr->Serialize(const_cast<void*>(Buffer.GetData()), static_cast<int64>(Buffer.GetSize()));
-			}
+			// Const cast because FArchive requires a non-const pointer!
+			FileAr->Serialize(const_cast<void*>(Buffer.GetData()), static_cast<int64>(Buffer.GetSize()));
 		}
 
-		return EPushResult::Success;
+		if (FileAr->Close())
+		{
+			return EPushResult::Success;
+		}
+		else
+		{
+			// TODO: If we were first saving to a tmp file we could avoid the need to delete the 
+			// potentially corrupt output file.
+			IFileManager::Get().Delete(FilePath.ToString()); 
+			return EPushResult::Failed;
+		}	
 	}
 
 	virtual FCompressedBuffer PullData(const FPayloadId& Id) override
@@ -109,19 +147,10 @@ private:
 		TUniquePtr<FArchive> FileAr(IFileManager::Get().CreateFileReader(FilePath.ToString()));
 		if (FileAr == nullptr)
 		{
-			const uint32 SystemError = FPlatformMisc::GetLastError();
-			// If we have a system error we can give a more informative error message but don't output it if the error is zero as 
-			// this can lead to very confusing error messages.
-			if (SystemError != 0)
-			{
-				TCHAR SystemErrorMsg[2048] = { 0 };
-				FPlatformMisc::GetSystemErrorMessage(SystemErrorMsg, sizeof(SystemErrorMsg), SystemError);
-				UE_LOG(LogVirtualization, Error, TEXT("[%s] Failed to load payload '%s' file '%s' due to system error: '%s' (%d))"), *GetDebugString(), *Id.ToString(), FilePath.ToString(), SystemErrorMsg, SystemError);
-			}
-			else
-			{
-				UE_LOG(LogVirtualization, Error, TEXT("[%s] Failed to load payload '%s' from '%s' (reason unknown)"), *GetDebugString(), *Id.ToString(), FilePath.ToString());
-			}
+			TStringBuilder<MAX_SPRINTF> SystemErrorMsg;
+			GetFormattedSystemError(SystemErrorMsg);
+
+			UE_LOG(LogVirtualization, Error, TEXT("[%s] Failed to load payload '%s' from file '%s' due to system error: %s"), *GetDebugString(), *Id.ToString(), FilePath.ToString(), SystemErrorMsg.ToString());
 		
 			return FCompressedBuffer();
 		}
diff --git a/Engine/Source/Developer/Virtualization/Private/VirtualizationJupiterBackend.cpp b/Engine/Source/Developer/Virtualization/Private/VirtualizationJupiterBackend.cpp
--- a/Engine/Source/Developer/Virtualization/Private/VirtualizationJupiterBackend.cpp
+++ b/Engine/Source/Developer/Virtualization/Private/VirtualizationJupiterBackend.cpp
@@ -1545,14 +1545,22 @@ private:
 				ResponseCode = Request->GetResponseCode();
 
 				// Request was successful, make sure we got all the expected data.
-				if (ResponseCode == 200)
+				if (FRequest::IsSuccessfulResponse(ResponseCode))
 				{
 					if (!Response.FromJson(Request->GetResponseAsJsonObject()))
 					{
-						UE_LOG(LogVirtualization, Error, TEXT("No valid response to http GET found, when accessing payload '%s'"), *Id.ToString());
+						UE_LOG(LogVirtualization, Error, TEXT("Failed to parser the header infomation about payload '%s'"), *Id.ToString());
 						return FCompressedBuffer();
 					}
 				}
+				else if (ResponseCode == 400)
+				{
+					// Response 400 indicates that the payload does not exist in Jupiter. Note that it is faster to just make the request
+					// and check for the response rather than call ::DoesPayloadExist prior to requesting the json header because this way 
+					// we will only make a single request if the payload exists or not.
+					UE_LOG(LogVirtualization, Verbose, TEXT("[%s] Does not contain the payload '%s'"), *GetDebugString(), *Id.ToString());
+					return FCompressedBuffer();
+				}
 				else if(!ShouldRetryOnError(ResponseCode))
 				{
 					UE_LOG(LogVirtualization, Error, TEXT("Failed with error code '%d' to download header infomation about payload '%s'"), ResponseCode, *Id.ToString());
diff --git a/Engine/Source/Developer/Virtualization/Private/VirtualizationSourceControlBackend.cpp b/Engine/Source/Developer/Virtualization/Private/VirtualizationSourceControlBackend.cpp
--- a/Engine/Source/Developer/Virtualization/Private/VirtualizationSourceControlBackend.cpp
+++ b/Engine/Source/Developer/Virtualization/Private/VirtualizationSourceControlBackend.cpp
@@ -131,13 +131,13 @@ public:
 		ISourceControlProvider& SCCProvider = ISourceControlModule::Get().GetProvider();
 
 #if IS_SOURCE_CONTROL_THREAD_SAFE
-		TSharedRef<FDownloadFile, ESPMode::ThreadSafe> DownloadCommand = ISourceControlOperation::Create<FDownloadFile>();
+		TSharedRef<FDownloadFile, ESPMode::ThreadSafe> DownloadCommand = ISourceControlOperation::Create<FDownloadFile>(FDownloadFile::EVerbosity::None);
 		if (SCCProvider.Execute(DownloadCommand, DepotPath.ToString(), EConcurrency::Synchronous) != ECommandResult::Succeeded)
 		{
 			return FCompressedBuffer();
 		}
 #else
-		TSharedRef<FDownloadFile> DownloadCommand = ISourceControlOperation::Create<FDownloadFile>();
+		TSharedRef<FDownloadFile> DownloadCommand = ISourceControlOperation::Create<FDownloadFile>(FDownloadFile::EVerbosity::None);
 		if (!SCCProvider.TryToDownloadFileFromBackgroundThread(DownloadCommand, DepotPath.ToString()))
 		{
 			return FCompressedBuffer();
