Increase the number of potential Mirage Jupiter requests which reduced the time spent pulling from Jupiter by about 30% on machines with more cores.

#rnx

- Increase the number of potential requests from 32 to 64 (64 to 96 would reduce the time spent by another 5% but that doesn't seem worth the cost for now)
- Change the index we use to iterate over the requests in the pool from a 8bit value to 32bit value to prevent easy overflows.

diff --git a/Engine/Source/Developer/Virtualization/Private/VirtualizationJupiterBackend.cpp b/Engine/Source/Developer/Virtualization/Private/VirtualizationJupiterBackend.cpp
--- a/Engine/Source/Developer/Virtualization/Private/VirtualizationJupiterBackend.cpp
+++ b/Engine/Source/Developer/Virtualization/Private/VirtualizationJupiterBackend.cpp
@@ -89,7 +89,7 @@ namespace Utility
 #define UE_MIRAGE_REQUEST_TIMEOUT_SECONDS 30L
 #define UE_MIRAGE_REQUEST_TIMOUT_ENABLED 1
 #define UE_MIRAGE_DEBUG 0
-#define UE_MIRAGE_REQUEST_POOL_SIZE 32
+#define UE_MIRAGE_REQUEST_POOL_SIZE 64
 #define UE_MIRAGE_MAX_FAILED_LOGIN_ATTEMPTS 16
 #define UE_MIRAGE_MAX_ATTEMPTS 4
 #define UE_MIRAGE_MAX_BUFFER_RESERVE 104857600u
@@ -877,20 +877,20 @@ struct FRequestPool
 {
 	FRequestPool(const TCHAR* InServiceUrl, FAccessToken* InAuthorizationToken)
 	{
-		for (uint8 i = 0; i < Pool.Num(); ++i)
+		for (int32 Index = 0; Index < Pool.Num(); ++Index)
 		{
-			Pool[i].Usage = 0u;
-			Pool[i].Request = new FRequest(InServiceUrl, InAuthorizationToken, true);	
+			Pool[Index].Usage = 0u;
+			Pool[Index].Request = new FRequest(InServiceUrl, InAuthorizationToken, true);	
 		}
 	}
 
 	~FRequestPool()
 	{
-		for (uint8 i = 0; i < Pool.Num(); ++i)
+		for (int32 Index = 0; Index < Pool.Num(); ++Index)
 		{
 			// No requests should be in use by now.
-			check(Pool[i].Usage.Load(EMemoryOrder::Relaxed) == 0u);
-			delete Pool[i].Request;
+			check(Pool[Index].Usage.Load(EMemoryOrder::Relaxed) == 0u);
+			delete Pool[Index].Request;
 		}
 	}
 
@@ -904,14 +904,14 @@ struct FRequestPool
 		TRACE_CPUPROFILER_EVENT_SCOPE(FRequestPool::WaitForFreeRequest);
 		while (true)
 		{
-			for (uint8 i = 0; i < Pool.Num(); ++i)
+			for (int32 Index = 0; Index < Pool.Num(); ++Index)
 			{
-				if (!Pool[i].Usage.Load(EMemoryOrder::Relaxed))
+				if (!Pool[Index].Usage.Load(EMemoryOrder::Relaxed))
 				{
 					uint8 Expected = 0u;
-					if (Pool[i].Usage.CompareExchange(Expected, 1u))
+					if (Pool[Index].Usage.CompareExchange(Expected, 1u))
 					{
-						return Pool[i].Request;
+						return Pool[Index].Request;
 					}
 				}
 			}
@@ -927,12 +927,12 @@ struct FRequestPool
 	*/
 	void ReleaseRequestToPool(FRequest* Request)
 	{
-		for (uint8 i = 0; i < Pool.Num(); ++i)
+		for (int32 Index = 0; Index < Pool.Num(); ++Index)
 		{
-			if (Pool[i].Request == Request)
+			if (Pool[Index].Request == Request)
 			{
 				Request->Reset();
-				Pool[i].Usage.Exchange(0u);
+				Pool[Index].Usage.Exchange(0u);
 				return;
 			}
 		}
@@ -2072,3 +2072,5 @@ UE_REGISTER_VIRTUALIZATION_BACKEND_FACTORY(FJupiterBackend, Jupiter);
 } // namespace UE
 
 #endif //WITH_MIRAGE_JUPITER_BACKEND
+
+PRAGMA_ENABLE_OPTIMIZATION
\ No newline at end of file
