CIS Fix - Remove PRAGMA_ENABLE_OPTIMIZATION from the bottom of the cpp.

#rnx

Unshelved the wrong version after preflight for submission

diff --git a/Engine/Source/Developer/Virtualization/Private/VirtualizationJupiterBackend.cpp b/Engine/Source/Developer/Virtualization/Private/VirtualizationJupiterBackend.cpp
--- a/Engine/Source/Developer/Virtualization/Private/VirtualizationJupiterBackend.cpp
+++ b/Engine/Source/Developer/Virtualization/Private/VirtualizationJupiterBackend.cpp
@@ -2072,5 +2072,3 @@ UE_REGISTER_VIRTUALIZATION_BACKEND_FACTORY(FJupiterBackend, Jupiter);
 } // namespace UE
 
 #endif //WITH_MIRAGE_JUPITER_BACKEND
-
-PRAGMA_ENABLE_OPTIMIZATION
\ No newline at end of file
