Allow Mirage to connect to the perforce source control provider even if it was not specified in the config file or command line.

#rnx

- When we check the source control provider name, we will now attempt to set it to "Perforce" if it is current the default "None" provider.
- If the user has specified a different source control provider than "Perforce" we will give a fatal error that the content virtualization graph containing the source control backend will not work with that setup.
- This part of the workflow might need to change in the future or at least the Mirage documentation must make very clear that when using source control as the persistent storage backend that a connection to perforce will be required, at least until we support our source control solutions but even then the connect to that provider will be required.

diff --git a/Engine/Source/Developer/Virtualization/Private/VirtualizationSourceControlBackend.cpp b/Engine/Source/Developer/Virtualization/Private/VirtualizationSourceControlBackend.cpp
--- a/Engine/Source/Developer/Virtualization/Private/VirtualizationSourceControlBackend.cpp
+++ b/Engine/Source/Developer/Virtualization/Private/VirtualizationSourceControlBackend.cpp
@@ -48,25 +48,24 @@ public:
 			return false;
 		}
 
-		const ISourceControlModule& SSCModule = ISourceControlModule::Get();
+		ISourceControlModule& SSCModule = ISourceControlModule::Get();
 
-		// We require source control the be enabled 
-		if (!SSCModule.IsEnabled())
+		// We require perforce as the source control provider as it is currently the only one that has the virtualization functionality implemented
+		const FName SourceControlName = SSCModule.GetProvider().GetName();
+		if (SourceControlName.IsNone())
 		{
-			UE_LOG(LogVirtualization, Error, TEXT("Attempting to initialize FSourceControlBackend but source control is disabled!"));
-			return false;
+			// No source control provider is set so we can try to set it to "Perforce"
+			// Note this call will fatal error if "Perforce" is not a valid option
+			SSCModule.SetProvider(FName("Perforce"));
 		}
-
-		ISourceControlProvider& SCCProvider = SSCModule.GetProvider();
-
-		// We require perforce as the source control provider as it is currently the only one that has the virtualization functionality implemented
-		const FName SourceControlName = SCCProvider.GetName();
-		if (SourceControlName != FName("Perforce"))
+		else if (SourceControlName != TEXT("Perforce"))
 		{
 			UE_LOG(LogVirtualization, Error, TEXT("Attempting to initialize FSourceControlBackend but source control is '%s' and only Perforce is currently supported!"), *SourceControlName.ToString());
 			return false;
 		}
 
+		ISourceControlProvider& SCCProvider = SSCModule.GetProvider();
+
 		if (!SCCProvider.IsAvailable())
 		{
 			SCCProvider.Init();
diff --git a/Engine/Source/Developer/Virtualization/Private/VirtualizationSourceControlUtilities.cpp b/Engine/Source/Developer/Virtualization/Private/VirtualizationSourceControlUtilities.cpp
--- a/Engine/Source/Developer/Virtualization/Private/VirtualizationSourceControlUtilities.cpp
+++ b/Engine/Source/Developer/Virtualization/Private/VirtualizationSourceControlUtilities.cpp
@@ -37,7 +37,7 @@ bool FVirtualizationSourceControlUtilities::SyncPayloadSidecarFile(const FPackag
 
 	// We only allow the perforce source control system as we have not adequately tested the others!
 	const FName SourceControlName = ISourceControlModule::Get().GetProvider().GetName();
-	if (SourceControlName != FName("Perforce"))
+	if (SourceControlName != TEXT("Perforce"))
 	{
 		UE_LOG(LogVirtualization, Error, TEXT("Attempting to sync a .upayload for '%s' but source control is '%s' and only Perforce is currently supported!"), *PackagePath.GetDebugName(), *SourceControlName.ToString());
 		return false;
