Reorganised the virtualization filesystem backend code.

#rnx

- Added a header file and moved the class declaration thereto make it easier for people to find the ini file documentation.
- Split some logging into multiple lines to reduce the overall width of the file.

diff --git a/Engine/Source/Developer/Virtualization/Private/VirtualizationFileBackend.cpp b/Engine/Source/Developer/Virtualization/Private/VirtualizationFileBackend.cpp
--- a/Engine/Source/Developer/Virtualization/Private/VirtualizationFileBackend.cpp
+++ b/Engine/Source/Developer/Virtualization/Private/VirtualizationFileBackend.cpp
@@ -1,12 +1,10 @@
 // Copyright Epic Games, Inc. All Rights Reserved.
 
-#include "Virtualization/IVirtualizationBackend.h"
+#include "VirtualizationFileBackend.h"
 
-#include "Containers/StringView.h"
 #include "HAL/FileManager.h"
 #include "Misc/Parse.h"
 #include "Misc/Paths.h"
-#include "Misc/StringBuilder.h"
 #include "Virtualization/VirtualizationManager.h"
 #include "VirtualizationUtilities.h"
 
@@ -39,153 +37,139 @@ void GetFormattedSystemError(FStringBuilderBase& SystemErrorMessage)
 	}
 }
 
-/**
- * A basic backend based on the file system. This can be used to access/store virtualization
- * data either on a local disk or a network share. It is intended to be used as a caching system
- * to speed up operations (running a local cache or a shared cache for a site) rather than as the
- * proper backend solution.
- * 
- * Ini file setup:
- * 'Name'=(Type=FileSystem, Path="XXX")
- * Where 'Name' is the backend name in the hierarchy and 'XXX' if the path to the directory where
- * you want to files to be stored.
- */
-class FFileSystemBackend : public IVirtualizationBackend
+
+FFileSystemBackend::FFileSystemBackend(FStringView ConfigName)
+	: IVirtualizationBackend(EOperations::Both)
 {
-public:
-	FFileSystemBackend(FStringView ConfigName)
-		: IVirtualizationBackend(EOperations::Both)
+	Name = WriteToString<256>(TEXT("FFileSystemBackend - "), ConfigName).ToString();
+}
+
+bool FFileSystemBackend::Initialize(const FString& ConfigEntry)
+{
+	if (!FParse::Value(*ConfigEntry, TEXT("Path="), RootDirectory))
 	{
-		Name = WriteToString<256>(TEXT("FFileSystemBackend - "), ConfigName).ToString();
+		UE_LOG(LogVirtualization, Error, TEXT("[%s] 'Path=' not found in the config file"), *GetDebugString());
+		return false;
 	}
 
-	virtual ~FFileSystemBackend() = default;
-
-private:
+	FPaths::NormalizeDirectoryName(RootDirectory);
 
-	virtual bool Initialize(const FString& ConfigEntry) override
+	if (RootDirectory.IsEmpty())
 	{
-		if (!FParse::Value(*ConfigEntry, TEXT("Path="), RootDirectory))
-		{
-			UE_LOG(LogVirtualization, Error, TEXT("[%s] 'Path=' not found in the config file"), *GetDebugString());
-			return false;
-		}
+		UE_LOG(LogVirtualization, Error, TEXT("[%s] Config file entry 'Path=' was empty"), *GetDebugString());
+		return false;
+	}
 
-		FPaths::NormalizeDirectoryName(RootDirectory);
+	// TODO: Validate that the given path is usable?
 
-		if (RootDirectory.IsEmpty())
-		{
-			UE_LOG(LogVirtualization, Error, TEXT("[%s] Config file entry 'Path=' was empty"), *GetDebugString());
-			return false;
-		}
+	UE_LOG(LogVirtualization, Log, TEXT("[%s] Using path: '%s'"), *GetDebugString(), *RootDirectory);
 
-		// TODO: Validate that the given path is usable?
+	return true;
+}
 
-		UE_LOG(LogVirtualization, Log, TEXT("[%s] Using path: '%s'"), *GetDebugString(), *RootDirectory);
+EPushResult FFileSystemBackend::PushData(const FPayloadId& Id, const FCompressedBuffer& Payload)
+{
+	TRACE_CPUPROFILER_EVENT_SCOPE(FFileSystemBackend::PushData);
 
-		return true;
+	if (DoesExist(Id))
+	{
+		UE_LOG(LogVirtualization, Verbose, TEXT("[%s] Already has a copy of the payload '%s'."), *GetDebugString(), *Id.ToString());
+		return EPushResult::PayloadAlreadyExisted;
 	}
 
-	virtual EPushResult PushData(const FPayloadId& Id, const FCompressedBuffer& Payload) override
+	TStringBuilder<512> FilePath;
+	CreateFilePath(Id, FilePath);
+
+	// TODO: Should we write to a temp file and then move it once it has written?
+	TUniquePtr<FArchive> FileAr(IFileManager::Get().CreateFileWriter(FilePath.ToString()));
+
+	if (FileAr == nullptr)
 	{
-		TRACE_CPUPROFILER_EVENT_SCOPE(FFileSystemBackend::PushData);
-
-		if (DoesExist(Id))
-		{
-			UE_LOG(LogVirtualization, Verbose, TEXT("[%s] Already has a copy of the payload '%s'."), *GetDebugString(), *Id.ToString());
-			return EPushResult::PayloadAlreadyExisted;
-		}
-
-		TStringBuilder<512> FilePath;
-		CreateFilePath(Id, FilePath);
-
-		// TODO: Should we write to a temp file and then move it once it has written?
-		TUniquePtr<FArchive> FileAr(IFileManager::Get().CreateFileWriter(FilePath.ToString()));
-
-		if (FileAr == nullptr)
-		{
-			TStringBuilder<MAX_SPRINTF> SystemErrorMsg;
-			GetFormattedSystemError(SystemErrorMsg);
-			UE_LOG(LogVirtualization, Error, TEXT("[%s] Failed to write payload '%s' to '%s' due to system error: %s"), *GetDebugString(), *Id.ToString(), FilePath.ToString(), SystemErrorMsg.ToString());
-
-			return EPushResult::Failed;
-		}
-
-		for (const FSharedBuffer& Buffer : Payload.GetCompressed().GetSegments())
-		{
-			// Const cast because FArchive requires a non-const pointer!
-			FileAr->Serialize(const_cast<void*>(Buffer.GetData()), static_cast<int64>(Buffer.GetSize()));
-		}
-
-		if (FileAr->Close())
-		{
-			return EPushResult::Success;
-		}
-		else
-		{
-			// TODO: If we were first saving to a tmp file we could avoid the need to delete the 
-			// potentially corrupt output file.
-			IFileManager::Get().Delete(FilePath.ToString()); 
-			return EPushResult::Failed;
-		}	
+		TStringBuilder<MAX_SPRINTF> SystemErrorMsg;
+		GetFormattedSystemError(SystemErrorMsg);
+
+		UE_LOG(LogVirtualization, Error, TEXT("[%s] Failed to write payload '%s' to '%s' due to system error: %s"), 
+				*GetDebugString(), 
+				*Id.ToString(), 
+				FilePath.ToString(), 
+				SystemErrorMsg.ToString());
+
+		return EPushResult::Failed;
 	}
 
-	virtual FCompressedBuffer PullData(const FPayloadId& Id) override
+	for (const FSharedBuffer& Buffer : Payload.GetCompressed().GetSegments())
 	{
-		TRACE_CPUPROFILER_EVENT_SCOPE(FFileSystemBackend::PullData);
-
-		TStringBuilder<512> FilePath;
-		CreateFilePath(Id, FilePath);
-
-		// TODO: Should we allow the error severity to be configured via ini or just not report this case at all?
-		if (!IFileManager::Get().FileExists(FilePath.ToString()))
-		{
-			UE_LOG(LogVirtualization, Verbose, TEXT("[%s] Does not contain the payload '%s'"), *GetDebugString(), *Id.ToString());
-			return FCompressedBuffer();
-		}
-
-		TUniquePtr<FArchive> FileAr(IFileManager::Get().CreateFileReader(FilePath.ToString()));
-		if (FileAr == nullptr)
-		{
-			TStringBuilder<MAX_SPRINTF> SystemErrorMsg;
-			GetFormattedSystemError(SystemErrorMsg);
-
-			UE_LOG(LogVirtualization, Error, TEXT("[%s] Failed to load payload '%s' from file '%s' due to system error: %s"), *GetDebugString(), *Id.ToString(), FilePath.ToString(), SystemErrorMsg.ToString());
-		
-			return FCompressedBuffer();
-		}
-
-		return FCompressedBuffer::FromCompressed(*FileAr);
+		// Const cast because FArchive requires a non-const pointer!
+		FileAr->Serialize(const_cast<void*>(Buffer.GetData()), static_cast<int64>(Buffer.GetSize()));
 	}
 
-	bool DoesExist(const FPayloadId& Id)
+	if (FileAr->Close())
+	{
+		return EPushResult::Success;
+	}
+	else
 	{
-		TRACE_CPUPROFILER_EVENT_SCOPE(FFileSystemBackend::DoesExist);
+		// TODO: If we were first saving to a tmp file we could avoid the need to delete the 
+		// potentially corrupt output file.
+		IFileManager::Get().Delete(FilePath.ToString()); 
+		return EPushResult::Failed;
+	}	
+}
 
-		TStringBuilder<512> FilePath;
-		CreateFilePath(Id, FilePath);
+FCompressedBuffer FFileSystemBackend::PullData(const FPayloadId& Id)
+{
+	TRACE_CPUPROFILER_EVENT_SCOPE(FFileSystemBackend::PullData);
 
-		return IFileManager::Get().FileExists(FilePath.ToString());
-	}
+	TStringBuilder<512> FilePath;
+	CreateFilePath(Id, FilePath);
 
-	virtual FString GetDebugString() const override
+	// TODO: Should we allow the error severity to be configured via ini or just not report this case at all?
+	if (!IFileManager::Get().FileExists(FilePath.ToString()))
 	{
-		return Name;
+		UE_LOG(LogVirtualization, Verbose, TEXT("[%s] Does not contain the payload '%s'"), *GetDebugString(), *Id.ToString());
+		return FCompressedBuffer();
 	}
 
-	void CreateFilePath(const FPayloadId& PayloadId, FStringBuilderBase& OutPath)
+	TUniquePtr<FArchive> FileAr(IFileManager::Get().CreateFileReader(FilePath.ToString()));
+	if (FileAr == nullptr)
 	{
-		TStringBuilder<52> PayloadPath;
-		Utils::PayloadIdToPath(PayloadId, PayloadPath);
-
-		OutPath << RootDirectory << TEXT("/") << PayloadPath;
+		TStringBuilder<MAX_SPRINTF> SystemErrorMsg;
+		GetFormattedSystemError(SystemErrorMsg);
+
+		UE_LOG(LogVirtualization, Error, TEXT("[%s] Failed to load payload '%s' from file '%s' due to system error: %s"), 
+				*GetDebugString(), 
+				*Id.ToString(), 
+				FilePath.ToString(), 
+				SystemErrorMsg.ToString());
+	
+		return FCompressedBuffer();
 	}
 
-private:
+	return FCompressedBuffer::FromCompressed(*FileAr);
+}
+
+FString FFileSystemBackend::GetDebugString() const
+{
+	return Name;
+}
 
-	FString Name;
-	FString RootDirectory;
-};
+bool FFileSystemBackend::DoesExist(const FPayloadId& Id)
+{
+	TRACE_CPUPROFILER_EVENT_SCOPE(FFileSystemBackend::DoesExist);
+
+	TStringBuilder<512> FilePath;
+	CreateFilePath(Id, FilePath);
+
+	return IFileManager::Get().FileExists(FilePath.ToString());
+}
+
+void FFileSystemBackend::CreateFilePath(const FPayloadId& PayloadId, FStringBuilderBase& OutPath)
+{
+	TStringBuilder<52> PayloadPath;
+	Utils::PayloadIdToPath(PayloadId, PayloadPath);
+
+	OutPath << RootDirectory << TEXT("/") << PayloadPath;
+}
 
 UE_REGISTER_VIRTUALIZATION_BACKEND_FACTORY(FFileSystemBackend, FileSystem);
 
diff --git a/Engine/Source/Developer/Virtualization/Private/VirtualizationFileBackend.h b/Engine/Source/Developer/Virtualization/Private/VirtualizationFileBackend.h
new file mode 100644
--- /dev/null
+++ b/Engine/Source/Developer/Virtualization/Private/VirtualizationFileBackend.h
@@ -0,0 +1,44 @@
+// Copyright Epic Games, Inc. All Rights Reserved.
+
+#pragma once
+
+#include "Virtualization/IVirtualizationBackend.h"
+
+namespace UE::Virtualization
+{
+/**
+ * A basic backend based on the file system. This can be used to access/store virtualization
+ * data either on a local disk or a network share. It is intended to be used as a caching system
+ * to speed up operations (running a local cache or a shared cache for a site) rather than as the
+ * proper backend solution.
+ *
+ * Ini file setup:
+ * 'Name'=(Type=FileSystem, Path="XXX")
+ * Where 'Name' is the backend name in the hierarchy and 'XXX' if the path to the directory where
+ * you want to files to be stored.
+ */
+class FFileSystemBackend : public IVirtualizationBackend
+{
+public:
+	FFileSystemBackend(FStringView ConfigName);
+	virtual ~FFileSystemBackend() = default;
+
+protected:
+
+	virtual bool Initialize(const FString& ConfigEntry) override;
+
+	virtual EPushResult PushData(const FPayloadId& Id, const FCompressedBuffer& Payload) override;
+
+	virtual FCompressedBuffer PullData(const FPayloadId& Id) override;
+
+	virtual FString GetDebugString() const override;
+
+	bool DoesExist(const FPayloadId& Id);
+
+	void CreateFilePath(const FPayloadId& PayloadId, FStringBuilderBase& OutPath);
+
+	FString Name;
+	FString RootDirectory;
+};
+
+} // namespace UE::Virtualization
\ No newline at end of file
