Add a new backend to Mirage that uses DDC2 as the storage mechanism.

#rnx

- This backend replaces the plans for a backend that interfaces with Zen directly.
- The backend has no real context for each payload being pushed at the moment, for now we pass in a dummy string.
- Currently we make blocking calls to the cache. When making the initial push to local storage this is desired as it will take place during package save. However both pulling the payload and making the submission to persistent storage should be able to be made async in the future.
- The storage policy is currently exposed via the ini file for now. We might want to enforce policies for this at a later time.

diff --git a/Engine/Source/Developer/Virtualization/Private/VirtualizationDDCBackend.cpp b/Engine/Source/Developer/Virtualization/Private/VirtualizationDDCBackend.cpp
new file mode 100644
--- /dev/null
+++ b/Engine/Source/Developer/Virtualization/Private/VirtualizationDDCBackend.cpp
@@ -0,0 +1,191 @@
+// Copyright Epic Games, Inc. All Rights Reserved.
+
+#include "VirtualizationDDCBackend.h"
+
+#include "Misc/Parse.h"
+#include "Virtualization/PayloadId.h"
+#include "DerivedDataCache.h"
+#include "DerivedDataCacheRecord.h"
+#include "DerivedDataRequestOwner.h"
+
+namespace UE::Virtualization
+{
+
+/** Utility function to help convert from UE::Virtualization types to UE::DerivedData::FPayload */
+UE::DerivedData::FPayload ToDDCPayload(const FPayloadId& Id, const FCompressedBuffer& Payload)
+{
+	UE::DerivedData::FPayloadId DDCPayloadId = UE::DerivedData::FPayloadId::FromHash(Id.GetIdentifier());
+
+	return UE::DerivedData::FPayload(DDCPayloadId, Payload);
+}
+
+FDDCBackend::FDDCBackend(FStringView ConfigName)
+: IVirtualizationBackend(EOperations::Both)
+, BucketName(TEXT("BulkData"))
+, TransferPolicy(UE::DerivedData::ECachePolicy::None)
+, QueryPolicy(UE::DerivedData::ECachePolicy::None)
+{
+	DebugName = WriteToString<256>(TEXT("FDDCBackend - "), ConfigName);
+}
+
+bool FDDCBackend::Initialize(const FString& ConfigEntry)
+{
+	TRACE_CPUPROFILER_EVENT_SCOPE(Initialize::Initialize);
+
+	if (!FParse::Value(*ConfigEntry, TEXT("Bucket="), BucketName))
+	{
+		UE_LOG(LogVirtualization, Fatal, TEXT("[%s] 'Bucket=' not found in the config file"), *GetDebugString());
+	}
+
+	bool bAllowLocal = true;
+	if (FParse::Bool(*ConfigEntry, TEXT("LocalStorage="), bAllowLocal))
+	{
+		UE_LOG(LogVirtualization, Log, TEXT("[%s] Use of local storage set to '%s"), *GetDebugString(), bAllowLocal ? TEXT("true") : TEXT("false"));
+	}
+
+	bool bAllowRemote = true;
+	if (FParse::Bool(*ConfigEntry, TEXT("RemoteStorage="), bAllowRemote))
+	{
+		UE_LOG(LogVirtualization, Log, TEXT("[%s] Use of remote storage set to '%s"), *GetDebugString(), bAllowRemote ? TEXT("true") : TEXT("false"));
+	}
+
+	if (!bAllowLocal && !bAllowRemote)
+	{
+		UE_LOG(LogVirtualization, Fatal, TEXT("[%s] LocalStorage and RemoteStorage cannot both be disabled"), *GetDebugString());
+		return false;
+	}
+
+	if (bAllowLocal)
+	{
+		TransferPolicy |= UE::DerivedData::ECachePolicy::Local;
+		QueryPolicy |= UE::DerivedData::ECachePolicy::QueryLocal;
+	}
+
+	if (bAllowRemote)
+	{
+		TransferPolicy |= UE::DerivedData::ECachePolicy::Remote;
+		QueryPolicy |= UE::DerivedData::ECachePolicy::QueryRemote;
+	}
+
+	return true;	
+}
+
+EPushResult FDDCBackend::PushData(const FPayloadId& Id, const FCompressedBuffer& Payload)
+{
+	TRACE_CPUPROFILER_EVENT_SCOPE(FDDCBackend::PushData);
+
+	if (DoesExist(Id))
+	{
+		UE_LOG(LogVirtualization, Verbose, TEXT("[%s] Already has a copy of the payload '%s'."), *GetDebugString(), *Id.ToString());
+		return EPushResult::PayloadAlreadyExisted;
+	}
+
+	UE::DerivedData::ICache& Cache = UE::DerivedData::GetCache();
+	
+	UE::DerivedData::FCacheKey Key;
+	Key.Bucket = UE::DerivedData::FCacheBucket(BucketName);
+	Key.Hash = Id.GetIdentifier();
+
+	UE::DerivedData::FPayload DDCPayload = ToDDCPayload(Id, Payload);
+	check(DDCPayload.GetRawHash() == Id.GetIdentifier());
+
+	UE::DerivedData::FCacheRecordBuilder RecordBuilder(Key);
+	RecordBuilder.SetValue(DDCPayload);
+
+	UE::DerivedData::FRequestOwner Owner(UE::DerivedData::EPriority::Blocking);
+
+	UE::DerivedData::FCachePutCompleteParams Result;
+	auto Callback = [&Result](UE::DerivedData::FCachePutCompleteParams&& Params)
+	{
+		Result = Params;
+	};
+
+	Cache.Put(	{RecordBuilder.Build()},
+				TEXT("Mirage"), // TODO: Improve this when we start passing more context to push
+				TransferPolicy,
+				Owner, 
+				MoveTemp(Callback));
+
+	Owner.Wait();
+
+	if (Result.Status == UE::DerivedData::EStatus::Ok)
+	{
+		return EPushResult::Success;
+	}
+	else
+	{
+		return EPushResult::Failed;
+	}
+}
+
+FCompressedBuffer FDDCBackend::PullData(const FPayloadId& Id)
+{
+	TRACE_CPUPROFILER_EVENT_SCOPE(FDDCBackend::PullData);
+
+	UE::DerivedData::ICache& Cache = UE::DerivedData::GetCache();
+
+	UE::DerivedData::FCacheKey Key;
+	Key.Bucket = UE::DerivedData::FCacheBucket(BucketName);
+	Key.Hash = Id.GetIdentifier();
+
+	UE::DerivedData::FRequestOwner Owner(UE::DerivedData::EPriority::Blocking);
+
+	FCompressedBuffer ResultData;
+	UE::DerivedData::EStatus ResultStatus;
+
+	auto Callback = [&ResultData, &ResultStatus](UE::DerivedData::FCacheGetCompleteParams&& Params)
+	{
+		ResultStatus = Params.Status;
+		if (ResultStatus == UE::DerivedData::EStatus::Ok)
+		{
+			ResultData = Params.Record.GetValuePayload().GetData();
+		}
+	};
+
+	Cache.Get(	{Key}, 
+				TEXT("Mirage"), // TODO: Improve this when we start passing more context to push
+				TransferPolicy,
+				Owner,
+				MoveTemp(Callback));
+
+	Owner.Wait();
+
+	return ResultData;
+}
+
+FString FDDCBackend::GetDebugString() const
+{
+	return DebugName;
+}
+
+bool FDDCBackend::DoesExist(const FPayloadId& Id) const
+{
+	TRACE_CPUPROFILER_EVENT_SCOPE(FDDCBackend::DoesExist);
+
+	UE::DerivedData::ICache& Cache = UE::DerivedData::GetCache();
+
+	UE::DerivedData::FCacheKey Key;
+	Key.Bucket = UE::DerivedData::FCacheBucket(BucketName);
+	Key.Hash = Id.GetIdentifier();
+
+	UE::DerivedData::FRequestOwner Owner(UE::DerivedData::EPriority::Blocking);
+	
+	UE::DerivedData::EStatus ResultStatus;
+	auto Callback = [&ResultStatus](UE::DerivedData::FCacheGetCompleteParams&& Params)
+	{
+		ResultStatus = Params.Status;
+	};
+
+	Cache.Get(	{ Key },
+				TEXT("Mirage"), // TODO: Improve this when we start passing more context to push
+				QueryPolicy | UE::DerivedData::ECachePolicy::SkipData,
+				Owner,
+				MoveTemp(Callback));
+
+	Owner.Wait();
+
+	return ResultStatus == UE::DerivedData::EStatus::Ok;
+}
+
+UE_REGISTER_VIRTUALIZATION_BACKEND_FACTORY(FDDCBackend, DDCBackend);
+} // namespace UE::Virtualization
\ No newline at end of file
diff --git a/Engine/Source/Developer/Virtualization/Private/VirtualizationDDCBackend.h b/Engine/Source/Developer/Virtualization/Private/VirtualizationDDCBackend.h
new file mode 100644
--- /dev/null
+++ b/Engine/Source/Developer/Virtualization/Private/VirtualizationDDCBackend.h
@@ -0,0 +1,56 @@
+// Copyright Epic Games, Inc. All Rights Reserved.
+
+#pragma once 
+
+#include "Virtualization/IVirtualizationBackend.h"
+
+namespace UE::DerivedData { enum class ECachePolicy : uint8; }
+
+namespace UE::Virtualization
+{
+/**
+ * A backend that uses the DDC2 as it's storage mechanism. It is intended to be used as a local caching
+ * system to speed up operations rather than for use as persistent storage.
+ * 
+ * Ini file setup:
+ * 'Name'=(Type=DDCBackend, Bucket="XXX", LocalStorage=True/False, RemoteStorage=True/False)
+ * 
+ * Required Values:
+ * 'Name': The backend name in the hierarchy.
+ * 
+ * Optional Values:
+ * Bucket: An alphanumeric identifier used to group the payloads together in storage. [Default="BulkData"]
+ * LocalStorage: When set to true, the payloads can be stored locally. [Default=true]
+ * RemoteStorage: When set to true, the payloads can be stored remotely. [Default=true]
+ */
+class FDDCBackend : public IVirtualizationBackend
+{
+public:
+	explicit FDDCBackend(FStringView ConfigName);
+	virtual ~FDDCBackend() = default;
+
+protected:
+
+	virtual bool Initialize(const FString& ConfigEntry) override;
+
+	virtual EPushResult PushData(const FPayloadId& Id, const FCompressedBuffer& Payload) override;
+
+	virtual FCompressedBuffer PullData(const FPayloadId& Id) override;
+
+	virtual FString GetDebugString() const override;
+
+	bool DoesExist(const FPayloadId& Id) const;
+	
+	/** Debug name used to identify the backend */
+	FString DebugName;
+
+	/** The bucket being used to group together the virtualized payloads in storage */
+	FString BucketName;
+
+	/** The policy to use when uploading or downloading data from the cache */
+	UE::DerivedData::ECachePolicy TransferPolicy;
+	/** The policy to use when querying the cache for information */
+	UE::DerivedData::ECachePolicy QueryPolicy;
+};
+
+} // namespace UE::Virtualization
\ No newline at end of file
diff --git a/Engine/Source/Developer/Virtualization/Virtualization.Build.cs b/Engine/Source/Developer/Virtualization/Virtualization.Build.cs
--- a/Engine/Source/Developer/Virtualization/Virtualization.Build.cs
+++ b/Engine/Source/Developer/Virtualization/Virtualization.Build.cs
@@ -10,7 +10,10 @@ public class Virtualization : ModuleRules
 		PrivateDependencyModuleNames.Add("Core");
 
 		// Remove this when we move the rest of the Virtualization code to this module
-		PrivateDependencyModuleNames.Add("CoreUObject"); 
+		PrivateDependencyModuleNames.Add("CoreUObject");
+
+		//  Dependency for DDC2
+		PrivateDependencyModuleNames.Add("DerivedDataCache");
 
 		// Dependencies for the Jupiter service backend
 		PrivateDependencyModuleNames.AddRange(new string[] { "SSL", "Json", "SourceControl" });
