Added a missing include of DerivedDataPayload.h

#rnx

diff --git a/Engine/Source/Developer/Virtualization/Private/VirtualizationDDCBackend.cpp b/Engine/Source/Developer/Virtualization/Private/VirtualizationDDCBackend.cpp
--- a/Engine/Source/Developer/Virtualization/Private/VirtualizationDDCBackend.cpp
+++ b/Engine/Source/Developer/Virtualization/Private/VirtualizationDDCBackend.cpp
@@ -6,6 +6,7 @@
 #include "Virtualization/PayloadId.h"
 #include "DerivedDataCache.h"
 #include "DerivedDataCacheRecord.h"
+#include "DerivedDataPayload.h"
 #include "DerivedDataRequestOwner.h"
 
 namespace UE::Virtualization
