Cache a FCacheBucket during the backend intialization and re-use it rather than recreating it for every request as the construction of this object is more complex than I thought.

#rnx

diff --git a/Engine/Source/Developer/Virtualization/Private/VirtualizationDDCBackend.cpp b/Engine/Source/Developer/Virtualization/Private/VirtualizationDDCBackend.cpp
--- a/Engine/Source/Developer/Virtualization/Private/VirtualizationDDCBackend.cpp
+++ b/Engine/Source/Developer/Virtualization/Private/VirtualizationDDCBackend.cpp
@@ -68,6 +68,8 @@ bool FDDCBackend::Initialize(const FString& ConfigEntry)
 		QueryPolicy |= UE::DerivedData::ECachePolicy::QueryRemote;
 	}
 
+	Bucket = UE::DerivedData::FCacheBucket(BucketName);
+
 	return true;	
 }
 
@@ -84,7 +86,7 @@ EPushResult FDDCBackend::PushData(const FPayloadId& Id, const FCompressedBuffer&
 	UE::DerivedData::ICache& Cache = UE::DerivedData::GetCache();
 	
 	UE::DerivedData::FCacheKey Key;
-	Key.Bucket = UE::DerivedData::FCacheBucket(BucketName);
+	Key.Bucket = Bucket;
 	Key.Hash = Id.GetIdentifier();
 
 	UE::DerivedData::FPayload DDCPayload = ToDDCPayload(Id, Payload);
@@ -126,7 +128,7 @@ FCompressedBuffer FDDCBackend::PullData(const FPayloadId& Id)
 	UE::DerivedData::ICache& Cache = UE::DerivedData::GetCache();
 
 	UE::DerivedData::FCacheKey Key;
-	Key.Bucket = UE::DerivedData::FCacheBucket(BucketName);
+	Key.Bucket = Bucket; 
 	Key.Hash = Id.GetIdentifier();
 
 	UE::DerivedData::FRequestOwner Owner(UE::DerivedData::EPriority::Blocking);
@@ -166,7 +168,7 @@ bool FDDCBackend::DoesExist(const FPayloadId& Id) const
 	UE::DerivedData::ICache& Cache = UE::DerivedData::GetCache();
 
 	UE::DerivedData::FCacheKey Key;
-	Key.Bucket = UE::DerivedData::FCacheBucket(BucketName);
+	Key.Bucket = Bucket;
 	Key.Hash = Id.GetIdentifier();
 
 	UE::DerivedData::FRequestOwner Owner(UE::DerivedData::EPriority::Blocking);
diff --git a/Engine/Source/Developer/Virtualization/Private/VirtualizationDDCBackend.h b/Engine/Source/Developer/Virtualization/Private/VirtualizationDDCBackend.h
--- a/Engine/Source/Developer/Virtualization/Private/VirtualizationDDCBackend.h
+++ b/Engine/Source/Developer/Virtualization/Private/VirtualizationDDCBackend.h
@@ -3,6 +3,7 @@
 #pragma once 
 
 #include "Virtualization/IVirtualizationBackend.h"
+#include "DerivedDataCacheKey.h"
 
 namespace UE::DerivedData { enum class ECachePolicy : uint8; }
 
@@ -47,6 +48,9 @@ protected:
 	/** The bucket being used to group together the virtualized payloads in storage */
 	FString BucketName;
 
+	/** */
+	UE::DerivedData::FCacheBucket Bucket;
+
 	/** The policy to use when uploading or downloading data from the cache */
 	UE::DerivedData::ECachePolicy TransferPolicy;
 	/** The policy to use when querying the cache for information */
