A quick hardening pass over the file system backend for Mirage.

#rnx

### Misc
- Updated the documentation for setting up a file system backend and added descriptions to the backend's members.
- Renamed the member 'Name' to 'DebugName' to bring it in line with the other backends (still a candidate to be moved to the base class)

### Reading payloads from disk
- Added the ability to retry opening file handles for read if they fail. This can commonly occur when multiple threads or processes are both pulling and pushing to the same root directory.
-- By default we will retry to open the file handle up to 10 times, waiting for 100ms between each attempt. The user can configure the number of retries and the length of the pause between each attempt via the config files.
-- Each failed attempt at opening the file handle will result in a warning being logged, so if the user is stalled for a long length of time while we rety over and over they are at least made aware of the cause.

### Writing payloads to disk
- When attempting to write a payload to disk, we will now write to a temp file in the project save directory using a randomly generated name. Once the payload has been fully saved to disk it will then be moved to the correct location. Although not perfect this will reduce the potential for writing corrupted data to a location where it might be read in the future.
-- We do attempt to clean up the temp file if the write fails but we do not print error messages if the delete fails as it is considered an optional bonus rather than an essential feature.
- Previously we would only log an error if we failed to open a file handle for writing the payload to disk. Now we also log an error if the handle is opened but the actual writes fail.

diff --git a/Engine/Source/Developer/Virtualization/Private/VirtualizationFileBackend.cpp b/Engine/Source/Developer/Virtualization/Private/VirtualizationFileBackend.cpp
--- a/Engine/Source/Developer/Virtualization/Private/VirtualizationFileBackend.cpp
+++ b/Engine/Source/Developer/Virtualization/Private/VirtualizationFileBackend.cpp
@@ -3,6 +3,7 @@
 #include "VirtualizationFileBackend.h"
 
 #include "HAL/FileManager.h"
+#include "HAL/PlatformProcess.h"
 #include "Misc/Parse.h"
 #include "Misc/Paths.h"
 #include "Virtualization/VirtualizationManager.h"
@@ -41,7 +42,7 @@ void GetFormattedSystemError(FStringBuilderBase& SystemErrorMessage)
 FFileSystemBackend::FFileSystemBackend(FStringView ConfigName)
 	: IVirtualizationBackend(EOperations::Both)
 {
-	Name = WriteToString<256>(TEXT("FFileSystemBackend - "), ConfigName).ToString();
+	DebugName = WriteToString<256>(TEXT("FFileSystemBackend - "), ConfigName).ToString();
 }
 
 bool FFileSystemBackend::Initialize(const FString& ConfigEntry)
@@ -62,7 +63,21 @@ bool FFileSystemBackend::Initialize(const FString& ConfigEntry)
 
 	// TODO: Validate that the given path is usable?
 
+	int32 RetryCountIniFile = INDEX_NONE;
+	if (FParse::Value(*ConfigEntry, TEXT("RetryCount="), RetryCountIniFile))
+	{
+		RetryCount = RetryCountIniFile;
+	}
+
+	int32 RetryWaitTimeMSIniFile = INDEX_NONE;
+	if (FParse::Value(*ConfigEntry, TEXT("RetryWaitTime="), RetryWaitTimeMSIniFile))
+	{
+		RetryWaitTimeMS = RetryWaitTimeMSIniFile;
+	}
+
+	// Now log a summary of the backend settings to make issues easier to diagnose
 	UE_LOG(LogVirtualization, Log, TEXT("[%s] Using path: '%s'"), *GetDebugString(), *RootDirectory);
+	UE_LOG(LogVirtualization, Log, TEXT("[%s] Will retry failed read attempts %d times with a gap of %dms betwen them"), *GetDebugString(), RetryCount, RetryWaitTimeMS);
 
 	return true;
 }
@@ -77,11 +92,14 @@ EPushResult FFileSystemBackend::PushData(const FPayloadId& Id, const FCompressed
 		return EPushResult::PayloadAlreadyExisted;
 	}
 
-	TStringBuilder<512> FilePath;
-	CreateFilePath(Id, FilePath);
-
-	// TODO: Should we write to a temp file and then move it once it has written?
-	TUniquePtr<FArchive> FileAr(IFileManager::Get().CreateFileWriter(FilePath.ToString()));
+	// Make sure to log any disk write failures to the user, even if this backend will often be optional as they are
+	// not expected and could indicate bigger problems.
+	// 
+	// First we will write out the payload to a temp file, after which we will move it to the correct storage location
+	// this helps reduce the chance of leaving corrupted data on disk in the case of a power failure etc.
+	const FString TempFilePath = FPaths::CreateTempFilename(*FPaths::ProjectSavedDir(), TEXT("miragepayload"));
+	
+	TUniquePtr<FArchive> FileAr(IFileManager::Get().CreateFileWriter(*TempFilePath));
 
 	if (FileAr == nullptr)
 	{
@@ -89,10 +107,10 @@ EPushResult FFileSystemBackend::PushData(const FPayloadId& Id, const FCompressed
 		GetFormattedSystemError(SystemErrorMsg);
 
 		UE_LOG(LogVirtualization, Error, TEXT("[%s] Failed to write payload '%s' to '%s' due to system error: %s"), 
-				*GetDebugString(), 
-				*Id.ToString(), 
-				FilePath.ToString(), 
-				SystemErrorMsg.ToString());
+			*GetDebugString(),
+			*Id.ToString(),
+			*TempFilePath,
+			SystemErrorMsg.ToString());
 
 		return EPushResult::Failed;
 	}
@@ -103,17 +121,53 @@ EPushResult FFileSystemBackend::PushData(const FPayloadId& Id, const FCompressed
 		FileAr->Serialize(const_cast<void*>(Buffer.GetData()), static_cast<int64>(Buffer.GetSize()));
 	}
 
-	if (FileAr->Close())
+	if (!FileAr->Close())
 	{
-		return EPushResult::Success;
+		TStringBuilder<MAX_SPRINTF> SystemErrorMsg;
+		GetFormattedSystemError(SystemErrorMsg);
+
+		UE_LOG(LogVirtualization, Error, TEXT("[%s] Failed to write payload '%s' contents to '%s' due to system error: %s"),
+			*GetDebugString(),
+			*Id.ToString(),
+			*TempFilePath,
+			SystemErrorMsg.ToString());
+
+		IFileManager::Get().Delete(*TempFilePath, true, false, true);  // Clean up the temp file if it is still around but do not failure cases to the user
+		
+		return EPushResult::Failed;
 	}
-	else
+
+	TStringBuilder<512> FilePath;
+	CreateFilePath(Id, FilePath);
+
+	if (!IFileManager::Get().Move(FilePath.ToString(), *TempFilePath))
 	{
-		// TODO: If we were first saving to a tmp file we could avoid the need to delete the 
-		// potentially corrupt output file.
-		IFileManager::Get().Delete(FilePath.ToString()); 
-		return EPushResult::Failed;
-	}	
+		// Store the error message in case we need to display it
+		TStringBuilder<MAX_SPRINTF> SystemErrorMsg;
+		GetFormattedSystemError(SystemErrorMsg);
+
+		IFileManager::Get().Delete(*TempFilePath, true, false, true); // Clean up the temp file if it is still around but do not failure cases to the user
+
+		// Check if another thread or process was writing out the payload at the same time, if so we 
+		// don't need to give an error message.
+		if (DoesExist(Id))
+		{
+			UE_LOG(LogVirtualization, Verbose, TEXT("[%s] Already has a copy of the payload '%s'."), *GetDebugString(), *Id.ToString());
+			return EPushResult::PayloadAlreadyExisted;
+		}
+		else
+		{
+			UE_LOG(LogVirtualization, Error, TEXT("[%s] Failed to move payload '%s' to it's final location '%s' due to system error: %s"),
+				*GetDebugString(),
+				*Id.ToString(),
+				*FilePath,
+				SystemErrorMsg.ToString());
+
+			return EPushResult::Failed;
+		}
+	}
+
+	return EPushResult::Success;
 }
 
 FCompressedBuffer FFileSystemBackend::PullData(const FPayloadId& Id)
@@ -130,18 +184,19 @@ FCompressedBuffer FFileSystemBackend::PullData(const FPayloadId& Id)
 		return FCompressedBuffer();
 	}
 
-	TUniquePtr<FArchive> FileAr(IFileManager::Get().CreateFileReader(FilePath.ToString()));
+	TUniquePtr<FArchive> FileAr = OpenFileForReading(FilePath.ToString());
+
 	if (FileAr == nullptr)
 	{
 		TStringBuilder<MAX_SPRINTF> SystemErrorMsg;
 		GetFormattedSystemError(SystemErrorMsg);
 
-		UE_LOG(LogVirtualization, Error, TEXT("[%s] Failed to load payload '%s' from file '%s' due to system error: %s"), 
-				*GetDebugString(), 
-				*Id.ToString(), 
-				FilePath.ToString(), 
-				SystemErrorMsg.ToString());
-	
+		UE_LOG(LogVirtualization, Error, TEXT("[%s] Failed to load payload '%s' from file '%s' due to system error: %s"),
+			*GetDebugString(),
+			*Id.ToString(),
+			FilePath.ToString(),
+			SystemErrorMsg.ToString());
+
 		return FCompressedBuffer();
 	}
 
@@ -150,7 +205,7 @@ FCompressedBuffer FFileSystemBackend::PullData(const FPayloadId& Id)
 
 FString FFileSystemBackend::GetDebugString() const
 {
-	return Name;
+	return DebugName;
 }
 
 bool FFileSystemBackend::DoesExist(const FPayloadId& Id)
@@ -171,6 +226,31 @@ void FFileSystemBackend::CreateFilePath(const FPayloadId& PayloadId, FStringBuil
 	OutPath << RootDirectory << TEXT("/") << PayloadPath;
 }
 
+TUniquePtr<FArchive> FFileSystemBackend::OpenFileForReading(const TCHAR* FilePath)
+{
+	TRACE_CPUPROFILER_EVENT_SCOPE(FFileSystemBackend::OpenFileForReading);
+
+	int32 Retries = 0;
+
+	while (Retries < RetryCount)
+	{
+		TUniquePtr<FArchive> FileAr(IFileManager::Get().CreateFileReader(FilePath));
+		if (FileAr)
+		{
+			return FileAr;
+		}
+		else
+		{
+			UE_LOG(LogVirtualization, Warning, TEXT("[%s] Failed to open '%s' for reading attempt retrying (%d/%d) in %dms..."), *GetDebugString(), FilePath, Retries, RetryCount, RetryWaitTimeMS);
+			FPlatformProcess::SleepNoStats(RetryWaitTimeMS * 0.001f);
+
+			Retries++;
+		}
+	}
+
+	return nullptr;
+}
+
 UE_REGISTER_VIRTUALIZATION_BACKEND_FACTORY(FFileSystemBackend, FileSystem);
 
 } // namespace UE::Virtualization
diff --git a/Engine/Source/Developer/Virtualization/Private/VirtualizationFileBackend.h b/Engine/Source/Developer/Virtualization/Private/VirtualizationFileBackend.h
--- a/Engine/Source/Developer/Virtualization/Private/VirtualizationFileBackend.h
+++ b/Engine/Source/Developer/Virtualization/Private/VirtualizationFileBackend.h
@@ -13,9 +13,20 @@ namespace UE::Virtualization
  * proper backend solution.
  *
  * Ini file setup:
- * 'Name'=(Type=FileSystem, Path="XXX")
- * Where 'Name' is the backend name in the hierarchy and 'XXX' if the path to the directory where
- * you want to files to be stored.
+ * 'Name'=(Type=FileSystem, Path="XXX", RetryCount=X, RetryWaitTime=X)
+ * 
+ * Required Values:
+ * 'Name': The backend name in the hierarchy.
+ * 'Type': The backend will be of type 'FFileSystemBackend'.
+ * 'Path': The root directory where the files are stored.
+ *
+ * Optional Values:
+ * RetryCount:		How many times we should try to open a payload file for read before giving up with
+ *					an error. Useful when many threads/processes can be pushing/pulling from the same
+ *					path/ [Default=10]
+ * RetryWaitTime:	The length of time the process should wait between each read attempt in milliseconds.
+ *					Remember that the max length of time that the process can stall attempting to read
+ *					a payload file is RetryCount * RetryWaitTime; [Default=100ms]
  */
 class FFileSystemBackend : public IVirtualizationBackend
 {
@@ -34,11 +45,20 @@ protected:
 	virtual FString GetDebugString() const override;
 
 	bool DoesExist(const FPayloadId& Id);
-
+	
 	void CreateFilePath(const FPayloadId& PayloadId, FStringBuilderBase& OutPath);
 
-	FString Name;
+	TUniquePtr<FArchive> OpenFileForReading(const TCHAR* FilePath);
+	
+	/** Debug name used to identify the backend. */
+	FString DebugName;
+	/** The root directory where the payload files should be located. */
 	FString RootDirectory;
+
+	/** The number of times to retry opening a payload file for read */
+	int32 RetryCount = 10;
+	/** The length of time (in milliseconds) to wait after each attempt before retrying. */
+	int32 RetryWaitTimeMS = 100;
 };
 
 } // namespace UE::Virtualization
\ No newline at end of file
