Add a console command 'DumpPackagePayloadInfo' that will take a list of package name or paths and then write info about each package's payloads to the log file.

#rnx

- This command is intended for debug use as an easy way to dump out the payload info in a package file that can then be inspected.

diff --git a/Engine/Source/Developer/Virtualization/Private/DumpPayloadTocCommand.cpp b/Engine/Source/Developer/Virtualization/Private/DumpPayloadTocCommand.cpp
new file mode 100644
--- /dev/null
+++ b/Engine/Source/Developer/Virtualization/Private/DumpPayloadTocCommand.cpp
@@ -0,0 +1,90 @@
+// Copyright Epic Games, Inc. All Rights Reserved.
+
+#include "HAL/IConsoleManager.h"
+#include "Logging/LogMacros.h"
+#include "Misc/PackagePath.h"
+#include "Virtualization/VirtualizedBulkData.h"
+
+namespace UE
+{
+namespace Virtualization
+{
+
+#if WITH_EDITORONLY_DATA
+
+/**
+ * This function is used to write information about package's payloads to the log file. This has no 
+ * practical development use and should only be used for debugging purposes. 
+ * 
+ * @param Args	The function expects each arg to be a valid package path. Failure to provide a valid
+ *				package path will result in errors being written to the log.
+ */
+void DumpPayloadToc(const TArray<FString>& Args)
+{
+	if (Args.Num() == 0)
+	{
+		UE_LOG(LogVirtualization, Error, TEXT("Command 'DumpPackagePayloadInfo' called without any arguments"));
+		return;
+	}
+
+	for (const FString& Arg : Args)
+	{
+		FPackagePath Path;
+
+		if (FPackagePath::TryFromMountedName(Arg, Path))
+		{
+			TArray<FPayloadId> LocalPayloadIds;
+			TArray<FPayloadId> VirtualizedPayloadIds;
+
+			if (!UE::Virtualization::FindPayloadsInPackageFile(Path, UE::Virtualization::EPayloadType::Local, LocalPayloadIds))
+			{
+				UE_LOG(LogVirtualization, Error, TEXT("Failed to find local payload information from package: '%s'"), *Path.GetDebugName());
+				continue;
+			}
+
+			if (!UE::Virtualization::FindPayloadsInPackageFile(Path, UE::Virtualization::EPayloadType::Virtualized, VirtualizedPayloadIds))
+			{
+				UE_LOG(LogVirtualization, Error, TEXT("Failed to find virtualized payload information from package: '%s'"), *Path.GetDebugName());
+				continue;
+			}
+
+			UE_LOG(LogVirtualization, Display, TEXT("Package: '%s' has %d local and %d virtualized payloads"), *Path.GetDebugName(), LocalPayloadIds.Num(), VirtualizedPayloadIds.Num());
+			
+			if (LocalPayloadIds.Num() > 0)
+			{
+				UE_LOG(LogVirtualization, Display, TEXT("LocalPayloads:"));
+				for (int32 Index = 0; Index < LocalPayloadIds.Num(); ++Index)
+				{
+					UE_LOG(LogVirtualization, Display, TEXT("%02d: '%s'"), Index, *LocalPayloadIds[Index].ToString());
+				}
+			}
+
+			if (VirtualizedPayloadIds.Num() > 0)
+			{
+				UE_LOG(LogVirtualization, Display, TEXT("VirtualizedPayloads:"));
+				for (int32 Index = 0; Index < VirtualizedPayloadIds.Num(); ++Index)
+				{
+					UE_LOG(LogVirtualization, Display, TEXT("%02d: '%s'"), Index, *VirtualizedPayloadIds[Index].ToString());
+				}
+			}
+		}
+		else
+		{
+			UE_LOG(LogVirtualization, Error, TEXT("Arg '%s' could not be converted to a valid package path"), *Arg);
+		}
+	}
+}
+
+/** 
+ * Note that this command is only valid when 'WITH_EDITORONLY_DATA 1' as virtualized payloads are not
+ * expected to exist at runtime. 
+ */
+static FAutoConsoleCommand CCmdDumpPayloadToc = FAutoConsoleCommand(
+	TEXT("DumpPackagePayloadInfo"),
+	TEXT("Writes out information about a package's payloads to the log."),
+	FConsoleCommandWithArgsDelegate::CreateStatic(DumpPayloadToc));
+
+#endif //WITH_EDITORONLY_DATA
+
+} // namespace Virtualization
+} // namespace UE
\ No newline at end of file
