Fix a problem that prevented some packages from being correctly virtualized on submit and clean up some of the surrounding code.

#rnx

- If packages were edited but not submitted and the editor shutdown then restarted there is a possibility that the edited packages will be loaded by the editor, meaning that a FLinkerLoader will have a file lock on the package file. So now we call ::ResetLoadersForSave on each package file that needs to be replaced to remove the file lock.
- When verbose logging is enabled we will now print how long the prepackagesubmission check took as a development aid.
- When packages have passed the prepackagesubmission process we now just append #virtualized to the CL description to mark the assets as having been through the process. The previous marker (a guid) was some what meaningless at the moment.
- Removed the legacy submit code, it is no longer needed for reference

diff --git a/Engine/Source/Developer/Virtualization/Private/PackageSubmissionChecks.cpp b/Engine/Source/Developer/Virtualization/Private/PackageSubmissionChecks.cpp
--- a/Engine/Source/Developer/Virtualization/Private/PackageSubmissionChecks.cpp
+++ b/Engine/Source/Developer/Virtualization/Private/PackageSubmissionChecks.cpp
@@ -5,94 +5,25 @@
 #include "Containers/UnrealString.h"
 #include "HAL/FileManager.h"
 #include "HAL/PlatformFileManager.h"
+#include "HAL/PlatformTime.h"
 #include "Internationalization/Internationalization.h"
 #include "Misc/PackageName.h"
 #include "Misc/Paths.h"
 #include "Serialization/VirtualizedBulkData.h"
+#include "UObject/Linker.h"
+#include "UObject/Package.h"
 #include "UObject/PackageResourceManager.h"
 #include "UObject/PackageTrailer.h"
+#include "UObject/UObjectGlobals.h"
 #include "Virtualization/VirtualizationSystem.h"
 
 #define LOCTEXT_NAMESPACE "Virtualization"
 
-#define UE_USE_LEGACY_SUBMIT 0
 #define UE_VALIDATE_TRUNCATED_PACKAGE 1
 
 namespace UE::Virtualization
 {
 
-#if UE_USE_LEGACY_SUBMIT
-void OnPrePackageSubmission(const TArray<FString>& FilesToSubmit, TArray<FText>& DescriptionTags, TArray<FText>& Errors)
-{
-	TRACE_CPUPROFILER_EVENT_SCOPE(UE::Virtualization::OnPrePackageSubmission);
-
-	IVirtualizationSystem& System = IVirtualizationSystem::Get();
-
-	// TODO: We could check to see if the package is virtualized even if it is disabled for the project
-	// as a safety feature?
-	if (!System.IsEnabled())
-	{
-		return;
-	}
-
-	// Other systems may have added errors to this array, we need to check so later we can determine if this function added any additional errors.
-	const int32 NumErrors = Errors.Num();
-
-	for (const FString& AbsoluteFilePath : FilesToSubmit)
-	{
-		const FPackagePath PackagePath = FPackagePath::FromLocalPath(AbsoluteFilePath);
-
-		if (FPackageName::IsPackageExtension(PackagePath.GetHeaderExtension()) || FPackageName::IsTextPackageExtension(PackagePath.GetHeaderExtension()))
-		{
-			TArray<FPayloadId> PayloadsInPackage;
-			if (FindPayloadsInPackageFile(PackagePath, EPayloadFilter::Virtualized, PayloadsInPackage))
-			{
-				for (const FPayloadId& PayloadId : PayloadsInPackage)
-				{
-					FCompressedBuffer Payload = System.PullData(PayloadId);
-
-					if (Payload)
-					{
-						if (!System.PushData(PayloadId, Payload, EStorageType::Persistent, PackagePath))
-						{
-							FText Message = FText::Format(LOCTEXT("Virtualization_PushFailure", "Failed to push payload '{0}' from the package '{1}'"),
-								FText::FromString(PayloadId.ToString()),
-								FText::FromString(PackagePath.GetDebugName()));
-
-							Errors.Add(Message);
-						}
-					}
-					else
-					{
-						FText Message = FText::Format(LOCTEXT("Virtualization_MissingPayload", "Unable to find the payload '{0}' for the package '{1}' in local storage"),
-							FText::FromString(PayloadId.ToString()),
-							FText::FromString(PackagePath.GetDebugName()));
-
-						Errors.Add(Message);
-					}
-				}
-			}
-			else
-			{
-				FText Message = FText::Format(LOCTEXT("Virtualization_PkgParseFailed", "Unable to find info about the payload requirements for the package '{0}'"),
-					FText::FromString(PackagePath.GetDebugName()));
-
-				Errors.Add(Message);
-			}
-		}
-	}
-
-	// If we had no new errors add the validation tag to indicate that the packages are safe for submission. 
-	// TODO: Currently this is a hard coded guid, it should be replaced by something unique to the FilesToSubmit collection that could be 
-	// validated by the source control server. This functionality is currently pending.
-	if (NumErrors == Errors.Num())
-	{
-		FText Tag = FText::FromString(TEXT("#virtualization c9c748e7-d00b-45af-bb61-9054abba503a"));
-		DescriptionTags.Add(Tag);	
-	}
-}
-#else
-
 /** 
  * Creates a copy of the given package but the copy will not include the FPackageTrailer.
  * 
@@ -188,6 +119,8 @@ void OnPrePackageSubmission(const TArray<FString>& FilesToSubmit, TArray<FText>&
 		return;
 	}
 
+	const double StartTime = FPlatformTime::Seconds();
+
 	// Other systems may have added errors to this array, we need to check so later we can determine if this function added any additional errors.
 	const int32 NumErrors = Errors.Num();
 
@@ -324,8 +257,17 @@ void OnPrePackageSubmission(const TArray<FString>& FilesToSubmit, TArray<FText>&
 		// TODO: Consider using the SavePackage model (move the original, then replace, so we can restore all of the original packages if needed)
 		// having said that, once a package is in PackagesToReplace it should still be safe to submit so maybe we don't need this level of protection?
 
-		// Since we had no errors we can now replace all of the packages that were virtualized data with the 
-		// virtualized replacement file.
+		// We need to reset the loader of any package that we want to re-save over
+		for (const TPair<FPackagePath, FString>& Iterator : PackagesToReplace)
+		{
+			UPackage* Package = FindObjectFast<UPackage>(nullptr, Iterator.Key.GetPackageFName());
+			if (Package != nullptr)
+			{
+				ResetLoadersForSave(Package, *Iterator.Key.GetLocalFullPath());
+			}
+		}
+
+		// Since we had no errors we can now replace all of the packages that were virtualized data with the virtualized replacement file.
 		for(const TPair<FPackagePath,FString>&  Iterator : PackagesToReplace)
 		{
 			const FString OriginalPackagePath = Iterator.Key.GetLocalFullPath();
@@ -342,15 +284,18 @@ void OnPrePackageSubmission(const TArray<FString>& FilesToSubmit, TArray<FText>&
 	}
 
 	// If we had no new errors add the validation tag to indicate that the packages are safe for submission. 
-	// TODO: Currently this is a hard coded guid, it should be replaced by something unique to the FilesToSubmit collection that could be 
-	// validated by the source control server. This functionality is currently pending.
+	// TODO: Currently this is a simple tag to make it easier for us to track which assets were submitted via the
+	// virtualization process in a test project. This should be expanded when we add proper p4 server triggers.
 	if (NumErrors == Errors.Num())
 	{
-		FText Tag = FText::FromString(TEXT("#virtualization c9c748e7-d00b-45af-bb61-9054abba503a"));
+		FText Tag = FText::FromString(TEXT("#virtualized"));
 		DescriptionTags.Add(Tag);
 	}
+
+	const double TimeInSeconds = FPlatformTime::ToSeconds64(FPlatformTime::Seconds() - StartTime);
+	UE_LOG(LogVirtualization, Verbose, TEXT("Virtualization pre submit check took %.3f(s)"), TimeInSeconds);
 }
-#endif 
+
 } // namespace UE::Virtualization
 
 #undef LOCTEXT_NAMESPACE
