Submitting packages on projects with virtualization enabled is much faster when none of the payloads actually needs to be virtualized.

#rnx

### PackageSubmissionChecks
- Under the old model submitting a large number of packages could be very slow as each package would check each payload that it owns and is currently stored locally one at a time. For the source control backend this created quite a large overhead even when all of the payloads were already virtualized.
- Now we do a pass over the submitted files to find all valid packages that have package trailers with locally stored payloads, gather the payloads into a single list and then query that in one large batch.
- Once we find which payloads are not in permanent storage (note that in the case where a project is using multiple permanent storage solutions, if a payload is missing in one backend it counts as not being in permanent storage) we then attempt to virtualized them.
- Only after all of this is done will we create the truncated copy of each package and then append the updated trailer to each one. In theory doing it in this order this might slightly increase the change of submit failures that occur after virtualization that result in a package never being submitted and orphaned payloads being added to permanent storage, but this will always be a risk.
- Added an assert to fire if we detect a trailer with some virtualized and some local payloads. This should be a supported feature but needs proper testing first before we can allow it. With out current project settings no project should actually encounter this scenario.
- To make the code easier to follow we now early out of the entire check when errors are encountered.
- Added logging at various stages in the process to help show the user that something is happening and make problems easier to identify in the future.
- Notes
--  There is a lot of handling of invalid FPayloads. This is because it is currently possible to add empty payloads to the trailer which is inefficient and wastes space. The trailer will be modified to reject empty payloads in a future update at which point a lot of this handling can be removed.
-- This could've also been solved by not fully rehydrating a package on save by the end user, which will be added as a project setting in a future piece of work, but this approach will solve the edge case when the user does have a large amount of hydrated packages which contain payloads that are already virtualized so it was better to fix that now while we have good test cases for it.
-- We still have scaling problems with large number of package being submitted that do have payloads that need to be virtualized, this will be fixed by extending IVirtualizationSystem::Push to also accept batches of payloads in future work.
-- OnPrePackageSubmission could be broken up into smaller chunks to make the code easier to follow. This will be done after the batch payload submission work is done.

### VirtualizationSystem
- EStorageType has been promoted to enum class.
- Added a new enum FPayloadStatus to be used when querying if a payload exists in a backend storage system or not.
- Add a new method ::DoPayloadsExist which allows the caller to query if one or more payloads exists in the given backend storage system.

### VirtualizationManager
- Implemented ::DoPayloadsExist. First we get the results from each backend in the storage system (which return as true or false from each backend) then total how many backends found the payload in order to set the correct status.

### IVirtualizationBackend
- ::DoesPayloadExist which queries the existence of a single payload has been added to the interface. Most backends already implemented this for private use and if so have had their implementation renamed to match this.
- Also added ::DoPayloadsExist which takes a batch of FpayloadIdsto query. Some backends can deal with a batch of payload ids much more efficiently than one at a time, although the default implementation does call ::DoesPayloadExist for each requested payload.
-- The default implementation prevents every backend from needing to implement the same for loop but does allow backends that can gain from batching to override it.

### VirtualizationSourceControlBackend
- This backend does override ::DoPayloadsExist and implements it's own version as it tends to perform very poorly when not operating on larger batches.
- In this case ::DoesPayloadExist calls back to ::DoPayloadsExist to check each payload rather than implement as specific version.

### PackageTrailer
- The trailer can now be queries to request how many payloads of a given type it contains

diff --git a/Engine/Source/Developer/Virtualization/Private/PackageSubmissionChecks.cpp b/Engine/Source/Developer/Virtualization/Private/PackageSubmissionChecks.cpp
--- a/Engine/Source/Developer/Virtualization/Private/PackageSubmissionChecks.cpp
+++ b/Engine/Source/Developer/Virtualization/Private/PackageSubmissionChecks.cpp
@@ -19,18 +19,57 @@
 
 #define LOCTEXT_NAMESPACE "Virtualization"
 
+// When enabled we will validate truncated packages right after the truncation process to 
+// make sure that the package format is still correct once the package trailer has been 
+// removed.
 #define UE_VALIDATE_TRUNCATED_PACKAGE 1
 
 namespace UE::Virtualization
 {
 
+/**
+ * Check that the given package ends with PACKAGE_FILE_TAG. Intended to be used to make sure that
+ * we have truncated a package correctly when removing the trailers.
+ * 
+ * @param PackagePath	The path of the package that should be checked
+ * @param Errors [out] 	Errors created by the function will be added here
+ * 
+ * @return	True if the package is correctly terminated with a PACKAGE_FILE_TAG, false if the tag
+ *			was not found or if we were unable to read the file's contents.
+ */
+bool ValidatePackage(const FString& PackagePath, TArray<FText>& Errors)
+{
+	TUniquePtr<IFileHandle> TempFileHandle(FPlatformFileManager::Get().GetPlatformFile().OpenRead(*PackagePath));
+	if (!TempFileHandle.IsValid())
+	{
+		FText ErrorMsg = FText::Format(LOCTEXT("Virtualization_OpenValidationFailed", "Unable to open '{0}' so that it can be validated"),
+			FText::FromString(PackagePath));
+		Errors.Add(ErrorMsg);
+		return false;
+	}
+
+	TempFileHandle->SeekFromEnd(-4);
+
+	uint32 PackageTag = INDEX_NONE;
+	if (!TempFileHandle->Read((uint8*)&PackageTag, 4) || PackageTag != PACKAGE_FILE_TAG)
+	{
+		FText ErrorMsg = FText::Format(LOCTEXT("Virtualization_ValidationFailed", "The package '{0}' does not end with a valid tag, the file is considered corrupt"),
+			FText::FromString(PackagePath));
+		Errors.Add(ErrorMsg);
+		return false;
+	}
+
+
+	return true;
+}
+
 /** 
  * Creates a copy of the given package but the copy will not include the FPackageTrailer.
  * 
- * @param[in]	PackagePath	The path of the package to copy
- * @param[in]	CopyPath	The path where the copy should be created
- * @param[in]	Trailer		The trailer found in 'PackagePath' that is already loaded
- * @param[out]	Errors		Errors created by the function will be added here
+ * @param PackagePath	The path of the package to copy
+ * @param CopyPath		The path where the copy should be created
+ * @param Trailer		The trailer found in 'PackagePath' that is already loaded
+ * @param Errors [out]	Errors created by the function will be added here
  * 
  * @return Returns true if the package was copied correctly, false otherwise. Note even when returning false a file might have been created at 'CopyPath'
  */
@@ -50,52 +89,32 @@ bool TryCopyPackageWithoutTrailer(const FPackagePath PackagePath, const FString&
 
 	const int64 PackageSizeWithoutTrailer = IFileManager::Get().FileSize(*PackageFilePath) - Trailer.GetTrailerLength();
 
-	TUniquePtr<IFileHandle> TempFileHandle(FPlatformFileManager::Get().GetPlatformFile().OpenWrite(*CopyPath, true));
-	if (!TempFileHandle.IsValid())
 	{
-		FText Message = FText::Format(	LOCTEXT("Virtualization_TruncOpenFailed", "Failed to open package file for truncation'{0}' when virtualizing"),
-										FText::FromString(CopyPath));
-		Errors.Add(Message);
-		return false;
-	}
-
-	if (!TempFileHandle->Truncate(PackageSizeWithoutTrailer))
-	{
-		FText Message = FText::Format(	LOCTEXT("Virtualization_TruncFailed", "Failed to truncate '{0}' when virtualizing"),
-										FText::FromString(CopyPath));
-		Errors.Add(Message);
-		return false;
-	}
-
-	return true;
-}
+		TUniquePtr<IFileHandle> TempFileHandle(FPlatformFileManager::Get().GetPlatformFile().OpenWrite(*CopyPath, true));
+		if (!TempFileHandle.IsValid())
+		{
+			FText Message = FText::Format(LOCTEXT("Virtualization_TruncOpenFailed", "Failed to open package file for truncation'{0}' when virtualizing"),
+				FText::FromString(CopyPath));
+			Errors.Add(Message);
+			return false;
+		}
 
-/** 
- * Check that the given package ends with PACKAGE_FILE_TAG. Intended to be used to make sure that
- * we have truncated a package correctly when removing the trailers. 
- */
-bool ValidatePackage(const FString& PackagePath, TArray<FText>& Errors)
-{
-	TUniquePtr<IFileHandle> TempFileHandle(FPlatformFileManager::Get().GetPlatformFile().OpenRead(*PackagePath));
-	if (!TempFileHandle.IsValid())
-	{
-		FText ErrorMsg = FText::Format(	LOCTEXT("Virtualization_OpenValidationFailed", "Unable to open '{0}' so that it can be validated"), 
-										FText::FromString(PackagePath));
-		Errors.Add(ErrorMsg);
-		return false;
+		if (!TempFileHandle->Truncate(PackageSizeWithoutTrailer))
+		{
+			FText Message = FText::Format(LOCTEXT("Virtualization_TruncFailed", "Failed to truncate '{0}' when virtualizing"),
+				FText::FromString(CopyPath));
+			Errors.Add(Message);
+			return false;
+		}
 	}
 
-	TempFileHandle->SeekFromEnd(-4);
-
-	uint32 PackageTag = INDEX_NONE;
-	if (!TempFileHandle->Read((uint8*)&PackageTag, 4) || PackageTag != PACKAGE_FILE_TAG)
+#if UE_VALIDATE_TRUNCATED_PACKAGE
+	// Validate we didn't break the package
+	if (!ValidatePackage(CopyPath, Errors))
 	{
-		FText ErrorMsg = FText::Format(	LOCTEXT("Virtualization_ValidationFailed", "The package '{0}' does not end with a valid tag, the file is considered corrupt"),
-										FText::FromString(PackagePath));
-		Errors.Add(ErrorMsg);
 		return false;
 	}
-
+#endif //UE_VALIDATE_TRUNCATED_PACKAGE
 
 	return true;
 }
@@ -124,134 +143,243 @@ void OnPrePackageSubmission(const TArray<FString>& FilesToSubmit, TArray<FText>&
 	// Other systems may have added errors to this array, we need to check so later we can determine if this function added any additional errors.
 	const int32 NumErrors = Errors.Num();
 
-	TArray<TPair<FPackagePath,FString>> PackagesToReplace;
+	struct FPackageInfo
+	{
+		FPackagePath Path;
+		FPackageTrailer Trailer;
+
+		TArray<FPayloadId> LocalPayloads;
+		int32 PayloadIndex = INDEX_NONE;
 
+		bool bWasTrailerUpdated = false;
+	};
+
+	UE_LOG(LogVirtualization, Display, TEXT("Considering %d files for virtualization"), FilesToSubmit.Num());
+
+	TArray<FPackageInfo> Packages;
+	Packages.Reserve(FilesToSubmit.Num());
+
+	TArray<FPayloadId> AllLocalPayloads;
+	AllLocalPayloads.Reserve(FilesToSubmit.Num());
+
+	// From the list of files to submit we need to find all of the valid packages that contain
+	// local payloads that need to be virtualized.
+	int64 TotalPayloadsToCheck = 0;
 	for (const FString& AbsoluteFilePath : FilesToSubmit)
 	{
-		if (NumErrors != Errors.Num())
-		{
-			break; // Give up if previous packages have had errors
-		}
-
-		const FPackagePath PackagePath = FPackagePath::FromLocalPath(AbsoluteFilePath);
+		FPackagePath PackagePath = FPackagePath::FromLocalPath(AbsoluteFilePath);
 
+		// TODO: How to handle text packages
 		if (FPackageName::IsPackageExtension(PackagePath.GetHeaderExtension()) || FPackageName::IsTextPackageExtension(PackagePath.GetHeaderExtension()))
 		{
 			FPackageTrailer Trailer;
 			if (FPackageTrailer::TryLoadFromPackage(PackagePath, Trailer))
 			{
-				const FString PackageFilePath = PackagePath.GetLocalFullPath();
-				const FString BaseName = FPaths::GetBaseFilename(PackagePath.GetPackageName());
-				const FString TempFilePath = FPaths::CreateTempFilename(*FPaths::ProjectSavedDir(), *BaseName.Left(32));
+				FPackageInfo PkgInfo;
 
-				// Create copy of package minus the trailer the trailer
-				if (!TryCopyPackageWithoutTrailer(PackagePath, TempFilePath, Trailer, Errors))
-				{
-					continue;
-				}			
+				PkgInfo.Path = MoveTemp(PackagePath);
+				PkgInfo.Trailer = MoveTemp(Trailer);
+				PkgInfo.LocalPayloads = PkgInfo.Trailer.GetPayloads(EPayloadFilter::Local);
 
-#if UE_VALIDATE_TRUNCATED_PACKAGE
-				// Validate we didn't break the package
-				if (!ValidatePackage(TempFilePath, Errors))
+				TotalPayloadsToCheck += PkgInfo.LocalPayloads.Num();
+
+				if (!PkgInfo.LocalPayloads.IsEmpty())
 				{
-					continue;
+					PkgInfo.PayloadIndex = AllLocalPayloads.Num();
+					AllLocalPayloads.Append(PkgInfo.LocalPayloads);
+
+					Packages.Emplace(MoveTemp(PkgInfo));
 				}
-#endif //UE_VALIDATE_TRUNCATED_PACKAGE
+			}
+		}
+	}
 
-				// Extract each payload and push
-				{
-					TUniquePtr<FArchive> PackageAr = IPackageResourceManager::Get().OpenReadExternalResource(EPackageExternalResource::WorkspaceDomainFile, PackagePath.GetPackageName());
+	UE_LOG(LogVirtualization, Display, TEXT("Found %" INT64_FMT " payloads in %d packages that need to be checked for virtualization"), TotalPayloadsToCheck, Packages.Num());
 
-					if (!PackageAr.IsValid())
-					{
-						FText Message = FText::Format(	LOCTEXT("Virtualization_PkgOpenm", "Failed to open the package '{1}' for reading"),
-														FText::FromString(PackagePath.GetDebugName()));
-						Errors.Add(Message);
-						break;
-					}
+	TArray<FPayloadStatus> PayloadStatuses;
+	if (!System.DoPayloadsExist(AllLocalPayloads, EStorageType::Persistent, PayloadStatuses))
+	{
+		FText Message = LOCTEXT("Virtualization_DoesExistFail", "Failed to find the status of the payloads in the packages being submitted");
+		Errors.Add(Message);
 
-					bool bUpdatedTrailer = false;
+		return;
+	}
 
-					for (const FPayloadId& PayloadId : Trailer.GetPayloads(EPayloadFilter::Local))
-					{
-						if (PayloadId.IsValid())
-						{
-							FCompressedBuffer Payload = Trailer.LoadPayload(PayloadId, *PackageAr);
-
-							if (PayloadId != FIoHash(Payload.GetRawHash()))
-							{
-								FText Message = FText::Format(LOCTEXT("Virtualization_WrongPayload", "Package {0} loaded an incorrect payload from the trailer. Expected '{1}' Loaded  '{2}'"),
-									FText::FromString(PackagePath.GetDebugName()),
-									FText::FromString(PayloadId.ToString()),
-									FText::FromString(LexToString(Payload.GetRawHash())));
-								Errors.Add(Message);
-								break;
-							}
-
-							if (Payload)
-							{
-								if (!System.PushData(PayloadId, Payload, EStorageType::Persistent, PackagePath))
-								{
-									FText Message = FText::Format(LOCTEXT("Virtualization_PushFailure", "Failed to push payload '{0}' from the package '{1}'"),
-										FText::FromString(PayloadId.ToString()),
-										FText::FromString(PackagePath.GetDebugName()));
-									Errors.Add(Message);
-									break;
-								}
-							}
-							else
-							{
-								FText Message = FText::Format(LOCTEXT("Virtualization_MissingPayload", "Unable to find the payload '{0}' in the local storage of package '{1}'"),
-									FText::FromString(PayloadId.ToString()),
-									FText::FromString(PackagePath.GetDebugName()));
-								Errors.Add(Message);
-								break;
-							}
-
-							if (Trailer.UpdatePayloadAsVirtualized(PayloadId))
-							{
-								bUpdatedTrailer = true;
-							}
-							else
-							{
-								FText Message = FText::Format(LOCTEXT("Virtualization_UpdateStatusFailed", "Unable to update the status for the payload '{0}' in the package '{1}'"),
-									FText::FromString(PayloadId.ToString()),
-									FText::FromString(PackagePath.GetDebugName()));
-								Errors.Add(Message);
-								break;
-							}		
-						}
-					}
-
-					if (bUpdatedTrailer)
+	// Update payloads that are already in persistent storage and don't need to be pushed
+	int64 TotalPayloadsToVirtualize = 0;
+	for (FPackageInfo& PackageInfo : Packages)
+	{
+		check(PackageInfo.LocalPayloads.IsEmpty() || PackageInfo.PayloadIndex != INDEX_NONE); // If we have payloads we should have an index
+
+		for (int32 Index = 0; Index < PackageInfo.LocalPayloads.Num(); ++Index)
+		{
+			// TODO: Remove FPayloadStatus::Invalid when trailers no longer support invalid payloads ids
+			if (PayloadStatuses[PackageInfo.PayloadIndex + Index] == FPayloadStatus::FoundAll ||
+				PayloadStatuses[PackageInfo.PayloadIndex + Index] == FPayloadStatus::Invalid)
+			{
+				if (PackageInfo.Trailer.UpdatePayloadAsVirtualized(PackageInfo.LocalPayloads[Index]))
+				{
+					PackageInfo.bWasTrailerUpdated = true;
+				}
+				else
 					{
-						TUniquePtr<FArchive> CopyAr(IFileManager::Get().CreateFileWriter(*TempFilePath, EFileWrite::FILEWRITE_Append));
-						if (!CopyAr.IsValid())
-						{
-							FText Message = FText::Format(	LOCTEXT("Virtualization_TrailerAppendOpen", "Unable to open '{0}' to append the trailer'"),
-															FText::FromString(TempFilePath));
-							Errors.Add(Message);
-							break;
-						}
-
-						FPackageTrailerBuilder TrailerBuilder = FPackageTrailerBuilder::Create(Trailer, *PackageAr);
-						if (!TrailerBuilder.BuildAndAppendTrailer(nullptr, *CopyAr))
-						{
-							FText Message = FText::Format(	LOCTEXT("Virtualization_TrailerAppend", "Failed to append the trailer to '{0}'"),
-															FText::FromString(TempFilePath));
-							Errors.Add(Message);
-							break;
-						}
-
-						// Now that we have successfully created a new version of the package with an updated trailer 
-						// we need to mark that it should replace the original package.
-						PackagesToReplace.Emplace(PackagePath, TempFilePath);
-					}
+					FText Message = FText::Format(	LOCTEXT("Virtualization_UpdateStatusFailed", "Unable to update the status for the payload '{0}' in the package '{1}'"),
+													FText::FromString(PackageInfo.LocalPayloads[Index].ToString()),
+													FText::FromString(PackageInfo.Path.GetDebugName()));
+					Errors.Add(Message);
+					return;
 				}
 			}
 		}
+
+		// If we made changes we should recalculate the local payloads left
+		if (PackageInfo.bWasTrailerUpdated)
+		{
+			PackageInfo.LocalPayloads = PackageInfo.Trailer.GetPayloads(EPayloadFilter::Local);
+		}
+
+		TotalPayloadsToVirtualize += PackageInfo.LocalPayloads.Num();
 	}
 
+	UE_LOG(LogVirtualization, Display, TEXT("Found %" INT64_FMT " payloads that need to be pushed to persistent virtualized storage"), TotalPayloadsToVirtualize);
+
+	// TODO Optimization: In theory we could have many packages sharing the same payload and we only need to push once
+
+	// Push any remaining local payload to the persistent backends
+	for (FPackageInfo& PackageInfo : Packages)
+	{
+		if (PackageInfo.LocalPayloads.IsEmpty())
+		{
+			continue;
+		}
+
+		TUniquePtr<FArchive> PackageAr = IPackageResourceManager::Get().OpenReadExternalResource(EPackageExternalResource::WorkspaceDomainFile, PackageInfo.Path.GetPackageName());
+
+		if (!PackageAr.IsValid())
+		{
+			FText Message = FText::Format(	LOCTEXT("Virtualization_PkgOpen", "Failed to open the package '{1}' for reading"),
+											FText::FromString(PackageInfo.Path.GetDebugName()));
+			Errors.Add(Message);
+			return;
+		}
+
+		for (const FPayloadId& PayloadId : PackageInfo.LocalPayloads)
+		{
+			check(PayloadId.IsValid());
+			
+			FCompressedBuffer Payload = PackageInfo.Trailer.LoadPayload(PayloadId, *PackageAr);
+
+			if (PayloadId != FIoHash(Payload.GetRawHash()))
+			{
+				FText Message = FText::Format(	LOCTEXT("Virtualization_WrongPayload", "Package {0} loaded an incorrect payload from the trailer. Expected '{1}' Loaded  '{2}'"),
+												FText::FromString(PackageInfo.Path.GetDebugName()),
+												FText::FromString(PayloadId.ToString()),
+												FText::FromString(LexToString(Payload.GetRawHash())));
+				Errors.Add(Message);
+				return;
+			}
+
+			if (Payload)
+			{
+				if (!System.PushData(PayloadId, Payload, EStorageType::Persistent, PackageInfo.Path))
+				{
+					FText Message = FText::Format(	LOCTEXT("Virtualization_PushFailure", "Failed to push payload '{0}' from the package '{1}'"),
+													FText::FromString(PayloadId.ToString()),
+													FText::FromString(PackageInfo.Path.GetDebugName()));
+					Errors.Add(Message);
+					return;
+				}
+			}
+			else
+			{
+				FText Message = FText::Format(	LOCTEXT("Virtualization_MissingPayload", "Unable to find the payload '{0}' in the local storage of package '{1}'"),
+												FText::FromString(PayloadId.ToString()),
+												FText::FromString(PackageInfo.Path.GetDebugName()));
+				Errors.Add(Message);
+				return;
+			}
+
+			if (PackageInfo.Trailer.UpdatePayloadAsVirtualized(PayloadId))
+			{
+				PackageInfo.bWasTrailerUpdated = true;
+			}
+			else
+			{
+				FText Message = FText::Format(	LOCTEXT("Virtualization_UpdateStatusFailed", "Unable to update the status for the payload '{0}' in the package '{1}'"),
+												FText::FromString(PayloadId.ToString()),
+												FText::FromString(PackageInfo.Path.GetDebugName()));
+				Errors.Add(Message);
+				return;
+			}
+		}
+	}
+
+	TArray<TPair<FPackagePath, FString>> PackagesToReplace;
+
+	// Any package with an updated trailer needs to be copied and an updated trailer appended
+	for (FPackageInfo& PackageInfo : Packages)
+	{
+		if (!PackageInfo.bWasTrailerUpdated)
+		{
+			continue;
+		}
+
+		// TODO: Implement/test mixed package trailer support
+		// At the moment we assume that all the payloads in a package will virtualized or none of them will. If none of the payloads have virtualized then
+		// the package won't have made it this far so we should check that no package at this point has a local payload.
+		// We should be able to support mixed packages (to support more advanced filtering and min payload sizes etc) but it hasn't been tested yet so I'd
+		// rather prevent this condition from occurring until it is properly tested.
+		check(PackageInfo.Trailer.GetNumPayloads(EPayloadFilter::Local) == 0);
+
+		const FPackagePath& PackagePath = PackageInfo.Path; // No need to validate path, we checked this earlier
+
+		const FString PackageFilePath = PackagePath.GetLocalFullPath();
+		const FString BaseName = FPaths::GetBaseFilename(PackagePath.GetPackageName());
+		const FString TempFilePath = FPaths::CreateTempFilename(*FPaths::ProjectSavedDir(), *BaseName.Left(32));
+
+		// TODO Optimization: Combine TryCopyPackageWithoutTrailer with the appending of the new trailer to avoid opening multiple handles
+
+		// Create copy of package minus the trailer the trailer
+		if (!TryCopyPackageWithoutTrailer(PackagePath, TempFilePath, PackageInfo.Trailer, Errors))
+		{
+			return;
+		}			
+
+		TUniquePtr<FArchive> PackageAr = IPackageResourceManager::Get().OpenReadExternalResource(EPackageExternalResource::WorkspaceDomainFile, PackagePath.GetPackageName());
+
+		if (!PackageAr.IsValid())
+		{
+			FText Message = FText::Format(	LOCTEXT("Virtualization_PkgOpen", "Failed to open the package '{1}' for reading"),
+											FText::FromString(PackagePath.GetDebugName()));
+			Errors.Add(Message);
+			return;
+		}
+
+		TUniquePtr<FArchive> CopyAr(IFileManager::Get().CreateFileWriter(*TempFilePath, EFileWrite::FILEWRITE_Append));
+		if (!CopyAr.IsValid())
+		{
+			FText Message = FText::Format(	LOCTEXT("Virtualization_TrailerAppendOpen", "Unable to open '{0}' to append the trailer'"),
+											FText::FromString(TempFilePath));
+			Errors.Add(Message);
+			return;
+		}
+
+		FPackageTrailerBuilder TrailerBuilder = FPackageTrailerBuilder::Create(PackageInfo.Trailer, *PackageAr);
+		if (!TrailerBuilder.BuildAndAppendTrailer(nullptr, *CopyAr))
+		{
+			FText Message = FText::Format(	LOCTEXT("Virtualization_TrailerAppend", "Failed to append the trailer to '{0}'"),
+											FText::FromString(TempFilePath));
+			Errors.Add(Message);
+			return;
+		}
+
+		// Now that we have successfully created a new version of the package with an updated trailer 
+		// we need to mark that it should replace the original package.
+		PackagesToReplace.Emplace(PackagePath, TempFilePath);
+	}
+
+	UE_LOG(LogVirtualization, Display, TEXT("%d packages had their trailer container modified and need to be updated"), PackagesToReplace.Num());
+
 	if (NumErrors == Errors.Num())
 	{
 		// TODO: Consider using the SavePackage model (move the original, then replace, so we can restore all of the original packages if needed)
diff --git a/Engine/Source/Developer/Virtualization/Private/VirtualizationDDCBackend.cpp b/Engine/Source/Developer/Virtualization/Private/VirtualizationDDCBackend.cpp
--- a/Engine/Source/Developer/Virtualization/Private/VirtualizationDDCBackend.cpp
+++ b/Engine/Source/Developer/Virtualization/Private/VirtualizationDDCBackend.cpp
@@ -77,7 +77,7 @@ EPushResult FDDCBackend::PushData(const FPayloadId& Id, const FCompressedBuffer&
 {
 	TRACE_CPUPROFILER_EVENT_SCOPE(FDDCBackend::PushData);
 
-	if (DoesExist(Id))
+	if (DoesPayloadExist(Id))
 	{
 		UE_LOG(LogVirtualization, Verbose, TEXT("[%s] Already has a copy of the payload '%s'."), *GetDebugName(), *Id.ToString());
 		return EPushResult::PayloadAlreadyExisted;
@@ -156,9 +156,9 @@ FCompressedBuffer FDDCBackend::PullData(const FPayloadId& Id)
 	return ResultData;
 }
 
-bool FDDCBackend::DoesExist(const FPayloadId& Id) const
+bool FDDCBackend::DoesPayloadExist(const FPayloadId& Id)
 {
-	TRACE_CPUPROFILER_EVENT_SCOPE(FDDCBackend::DoesExist);
+	TRACE_CPUPROFILER_EVENT_SCOPE(FDDCBackend::DoesPayloadExist);
 
 	UE::DerivedData::ICache& Cache = UE::DerivedData::GetCache();
 
diff --git a/Engine/Source/Developer/Virtualization/Private/VirtualizationDDCBackend.h b/Engine/Source/Developer/Virtualization/Private/VirtualizationDDCBackend.h
--- a/Engine/Source/Developer/Virtualization/Private/VirtualizationDDCBackend.h
+++ b/Engine/Source/Developer/Virtualization/Private/VirtualizationDDCBackend.h
@@ -37,9 +37,9 @@ protected:
 	virtual EPushResult PushData(const FPayloadId& Id, const FCompressedBuffer& Payload, const FPackagePath& PackageContext) override;
 
 	virtual FCompressedBuffer PullData(const FPayloadId& Id) override;
-
-	bool DoesExist(const FPayloadId& Id) const;
-
+	
+	virtual bool DoesPayloadExist(const FPayloadId& Id) override;
+	
 	/** The bucket being used to group together the virtualized payloads in storage */
 	FString BucketName;
 
diff --git a/Engine/Source/Developer/Virtualization/Private/VirtualizationFileBackend.cpp b/Engine/Source/Developer/Virtualization/Private/VirtualizationFileBackend.cpp
--- a/Engine/Source/Developer/Virtualization/Private/VirtualizationFileBackend.cpp
+++ b/Engine/Source/Developer/Virtualization/Private/VirtualizationFileBackend.cpp
@@ -58,7 +58,7 @@ EPushResult FFileSystemBackend::PushData(const FPayloadId& Id, const FCompressed
 {
 	TRACE_CPUPROFILER_EVENT_SCOPE(FFileSystemBackend::PushData);
 
-	if (DoesExist(Id))
+	if (DoesPayloadExist(Id))
 	{
 		UE_LOG(LogVirtualization, Verbose, TEXT("[%s] Already has a copy of the payload '%s'."), *GetDebugName(), *Id.ToString());
 		return EPushResult::PayloadAlreadyExisted;
@@ -123,7 +123,7 @@ EPushResult FFileSystemBackend::PushData(const FPayloadId& Id, const FCompressed
 
 		// Check if another thread or process was writing out the payload at the same time, if so we 
 		// don't need to give an error message.
-		if (DoesExist(Id))
+		if (DoesPayloadExist(Id))
 		{
 			UE_LOG(LogVirtualization, Verbose, TEXT("[%s] Already has a copy of the payload '%s'."), *GetDebugName(), *Id.ToString());
 			return EPushResult::PayloadAlreadyExisted;
@@ -176,9 +176,9 @@ FCompressedBuffer FFileSystemBackend::PullData(const FPayloadId& Id)
 	return FCompressedBuffer::FromCompressed(*FileAr);
 }
 
-bool FFileSystemBackend::DoesExist(const FPayloadId& Id)
+bool FFileSystemBackend::DoesPayloadExist(const FPayloadId& Id)
 {
-	TRACE_CPUPROFILER_EVENT_SCOPE(FFileSystemBackend::DoesExist);
+	TRACE_CPUPROFILER_EVENT_SCOPE(FFileSystemBackend::DoesPayloadExist);
 
 	TStringBuilder<512> FilePath;
 	CreateFilePath(Id, FilePath);
diff --git a/Engine/Source/Developer/Virtualization/Private/VirtualizationFileBackend.h b/Engine/Source/Developer/Virtualization/Private/VirtualizationFileBackend.h
--- a/Engine/Source/Developer/Virtualization/Private/VirtualizationFileBackend.h
+++ b/Engine/Source/Developer/Virtualization/Private/VirtualizationFileBackend.h
@@ -42,7 +42,7 @@ protected:
 
 	virtual FCompressedBuffer PullData(const FPayloadId& Id) override;
 
-	bool DoesExist(const FPayloadId& Id);
+	virtual bool DoesPayloadExist(const FPayloadId& Id) override;
 	
 	void CreateFilePath(const FPayloadId& PayloadId, FStringBuilderBase& OutPath);
 
diff --git a/Engine/Source/Developer/Virtualization/Private/VirtualizationJupiterBackend.cpp b/Engine/Source/Developer/Virtualization/Private/VirtualizationJupiterBackend.cpp
--- a/Engine/Source/Developer/Virtualization/Private/VirtualizationJupiterBackend.cpp
+++ b/Engine/Source/Developer/Virtualization/Private/VirtualizationJupiterBackend.cpp
@@ -1955,7 +1955,7 @@ private:
 		return false;
 	}
 
-	bool DoesPayloadExist(const FPayloadId& Id)
+	virtual bool DoesPayloadExist(const FPayloadId& Id) override
 	{
 		using namespace Utility;
 
diff --git a/Engine/Source/Developer/Virtualization/Private/VirtualizationManager.cpp b/Engine/Source/Developer/Virtualization/Private/VirtualizationManager.cpp
--- a/Engine/Source/Developer/Virtualization/Private/VirtualizationManager.cpp
+++ b/Engine/Source/Developer/Virtualization/Private/VirtualizationManager.cpp
@@ -426,6 +426,70 @@ FCompressedBuffer FVirtualizationManager::PullData(const FPayloadId& Id)
 	return FCompressedBuffer();
 }
 
+bool FVirtualizationManager::DoPayloadsExist(TArrayView<const FPayloadId> Ids, EStorageType StorageType, TArray<FPayloadStatus>& OutStatuses)
+{
+	TRACE_CPUPROFILER_EVENT_SCOPE(FVirtualizationManager::DoPayloadsExist);
+
+	OutStatuses.SetNum(Ids.Num());
+
+	for (int32 Index = 0; Index < Ids.Num(); ++Index)
+	{
+		OutStatuses[Index] = Ids[Index].IsValid() ? FPayloadStatus::NotFound : FPayloadStatus::Invalid;
+	}
+
+	FBackendArray& Backends = StorageType == EStorageType::Local ? LocalCachableBackends : PersistentStorageBackends;
+
+	TArray<int8> HitCount;
+	TArray<bool> Results;
+
+	HitCount.SetNum(Ids.Num());
+	Results.SetNum(Ids.Num());
+
+	{
+		FConditionalScopeLock _(&ForceSingleThreadedCS, bForceSingleThreaded);
+
+		for (IVirtualizationBackend* Backend : Backends)
+		{
+			if (!Backend->DoPayloadsExist(Ids, Results))
+			{
+				// If a backend entirely failed we should early out and report the problem
+				OutStatuses.Reset();
+				return false;
+			}
+
+			for (int32 Index = 0; Index < Ids.Num(); ++Index)
+			{
+				if (Ids[Index].IsValid() && Results[Index])
+				{
+					HitCount[Index]++;
+				}
+			}
+		}
+	}
+
+	// Now we total up the hit count for each payload to see if it was found in none, all or some of the backends
+	for (int32 Index = 0; Index < Ids.Num(); ++Index)
+	{
+		if (Ids[Index].IsValid())
+		{
+			if (HitCount[Index] == 0)
+			{
+				OutStatuses[Index] = FPayloadStatus::NotFound;
+			}
+			else if (HitCount[Index] == Backends.Num())
+			{
+				OutStatuses[Index] = FPayloadStatus::FoundAll;
+			}
+			else
+			{
+				OutStatuses[Index] = FPayloadStatus::Partial;;
+			}
+		}
+	}
+
+	return true;
+}
+
 FPayloadActivityInfo FVirtualizationManager::GetAccumualtedPayloadActivityInfo() const
 {
 	FPayloadActivityInfo Info;
diff --git a/Engine/Source/Developer/Virtualization/Private/VirtualizationManager.h b/Engine/Source/Developer/Virtualization/Private/VirtualizationManager.h
--- a/Engine/Source/Developer/Virtualization/Private/VirtualizationManager.h
+++ b/Engine/Source/Developer/Virtualization/Private/VirtualizationManager.h
@@ -118,13 +118,25 @@ public:
 	 */
 	virtual FCompressedBuffer PullData(const FPayloadId& Id) override;
 
+	/**
+	 * Query if a number of payloads exist or not in the given storage type.
+	 *
+	 * @param	Ids					One or more payload identifiers to test
+	 * @param	StorageType			The type of storage to push the payload to, @See EStorageType for details.
+	 * @param	OutStatuses [out]	An array containing the results for each payload. @See FPayloadStatus
+	 * 								If the operation succeeds the array will be resized to match the size of Ids.
+	 *
+	 * @return	True if the operation succeeded and the contents of OutStatuses is valid. False if errors were
+	 * 			encountered in which case the contents of OutStatuses should be ignored.
+	 */
+	virtual bool DoPayloadsExist(TArrayView<const FPayloadId> Ids, EStorageType StorageType, TArray<FPayloadStatus>& OutStatuses) override;
+
 	/** Access profiling info relating to accumulated payload activity for all backends. Stats will only be collected if ENABLE_COOK_STATS is enabled.*/
 	virtual FPayloadActivityInfo GetAccumualtedPayloadActivityInfo() const override;
 
 	/** Access profiling info relating to payload activity per backend. Stats will only be collected if ENABLE_COOK_STATS is enabled.*/
 	virtual void GetPayloadActivityInfo( GetPayloadActivityInfoFuncRef ) const override;
 
-
 	/** Get event delegate for system notifications */
 	virtual FOnNotification& GetNotificationEvent() override
 	{
diff --git a/Engine/Source/Developer/Virtualization/Private/VirtualizationSourceControlBackend.cpp b/Engine/Source/Developer/Virtualization/Private/VirtualizationSourceControlBackend.cpp
--- a/Engine/Source/Developer/Virtualization/Private/VirtualizationSourceControlBackend.cpp
+++ b/Engine/Source/Developer/Virtualization/Private/VirtualizationSourceControlBackend.cpp
@@ -14,6 +14,7 @@
 #include "VirtualizationSourceControlUtilities.h"
 #include "VirtualizationUtilities.h"
 
+
 // When the SourceControl module (or at least the perforce source control module) is thread safe we
 // can enable this and stop using the hacky work around 'TryToDownloadFileFromBackgroundThread'
 #define IS_SOURCE_CONTROL_THREAD_SAFE 0
@@ -338,6 +339,64 @@ public:
 		OutPath << DepotRoot << PayloadPath;
 	}
 
+	virtual bool DoesPayloadExist(const FPayloadId& Id)
+	{
+		TArray<bool> Result;
+
+		if (FSourceControlBackend::DoPayloadsExist(MakeArrayView<const FPayloadId>(&Id, 1), Result))
+		{
+			check(Result.Num() == 1);
+			return Result[0];
+		}
+		else
+		{
+			return false;
+		}
+	}
+
+	virtual bool DoPayloadsExist(TArrayView<const FPayloadId> PayloadIds, TArray<bool>& OutResults) override
+	{
+		ISourceControlProvider& SCCProvider = ISourceControlModule::Get().GetProvider();
+
+		TArray<FString> DepotPaths;
+		DepotPaths.Reserve(PayloadIds.Num());
+
+		TArray<FSourceControlStateRef> PathStates;
+	
+		for (const FPayloadId& PayloadId : PayloadIds)
+		{
+			if (PayloadId.IsValid())
+			{
+				TStringBuilder<52> LocalPayloadPath;
+				Utils::PayloadIdToPath(PayloadId, LocalPayloadPath);
+
+				DepotPaths.Emplace(WriteToString<512>(DepotRoot, LocalPayloadPath));
+			}		
+		}
+	
+		ECommandResult::Type Result = SCCProvider.GetState(DepotPaths, PathStates, EStateCacheUsage::ForceUpdate);
+		if (Result != ECommandResult::Type::Succeeded)
+		{
+			UE_LOG(LogVirtualization, Error, TEXT("[%s] Failed to query the state of files in the source control depot"), *GetDebugName());
+			return false;
+		}
+
+		check(DepotPaths.Num() == PathStates.Num()); // We expect that all paths return a state
+
+		OutResults.SetNum(PayloadIds.Num());
+
+		int32 StatusIndex = 0;
+		for (int32 Index = 0; Index < PayloadIds.Num(); ++Index)
+		{
+			if (PayloadIds[Index].IsValid())
+			{
+				OutResults[Index] = PathStates[StatusIndex++]->IsSourceControlled();
+			}
+		}
+
+		return true;
+	}
+
 private:
 
 	/** The root where the virtualized payloads are stored in source control */
diff --git a/Engine/Source/Developer/Virtualization/Public/IVirtualizationBackend.h b/Engine/Source/Developer/Virtualization/Public/IVirtualizationBackend.h
--- a/Engine/Source/Developer/Virtualization/Public/IVirtualizationBackend.h
+++ b/Engine/Source/Developer/Virtualization/Public/IVirtualizationBackend.h
@@ -100,7 +100,43 @@ public:
 	 *				if it did not.
 	 */
 	virtual FCompressedBuffer PullData(const FPayloadId& Id) = 0;
+	
+	/**
+	 * Checks if a payload exists in the backends storage.
+	 * 
+	 * @param	Id	The identifier of the payload to check
+	 * 
+	 * @return True if the backend storage already contains the payload, otherwise false
+	 */
+	virtual bool DoesPayloadExist(const FPayloadId& Id) = 0;
+	
+	/**
+	 * Checks if a number of payload exists in the backends storage.
+	 *
+	 * @param[in]	PayloadIds	An array of FPayloadId that should be checked
+	 * @param[out]	OutResults	An array to contain the result, true if the payload
+	 *							exists in the backends storage, false if not.
+	 *							This array will be resized to match the size of PayloadIds.
+	 * 
+	 * @return True if the operation completed without error, otherwise false
+	 */
+	virtual bool DoPayloadsExist(TArrayView<const FPayloadId> PayloadIds, TArray<bool>& OutResults)
+	{
+		// This is the default implementation that just calls ::DoesExist on each FPayloadId in the
+		// array, one at a time. 
+		// Backends may override this with their own implementations if it can be done with less
+		// overhead by performing the check on the entire batch instead.
+
+		OutResults.SetNum(PayloadIds.Num());
+
+		for (int32 Index = 0; Index < PayloadIds.Num(); ++Index)
+		{
+			OutResults[Index] = DoesPayloadExist(PayloadIds[Index]);
+		}
 
+		return true;
+	}
+	
 	/** Used when debugging to disable the pull operation */
 	void DisablePullOperationSupport()
 	{
diff --git a/Engine/Source/Runtime/Core/Private/Virtualization/VirtualizationSystem.cpp b/Engine/Source/Runtime/Core/Private/Virtualization/VirtualizationSystem.cpp
--- a/Engine/Source/Runtime/Core/Private/Virtualization/VirtualizationSystem.cpp
+++ b/Engine/Source/Runtime/Core/Private/Virtualization/VirtualizationSystem.cpp
@@ -34,6 +34,13 @@ public:
 		return FCompressedBuffer();
 	}
 
+	virtual bool DoPayloadsExist(TArrayView<const FPayloadId> Ids, EStorageType StorageType, TArray<FPayloadStatus>& OutStatuses) override
+	{
+		OutStatuses.Reset();
+
+		return false;
+	}
+
 	virtual void GetPayloadActivityInfo(GetPayloadActivityInfoFuncRef) const override
 	{
 	}
diff --git a/Engine/Source/Runtime/Core/Public/Virtualization/VirtualizationSystem.h b/Engine/Source/Runtime/Core/Public/Virtualization/VirtualizationSystem.h
--- a/Engine/Source/Runtime/Core/Public/Virtualization/VirtualizationSystem.h
+++ b/Engine/Source/Runtime/Core/Public/Virtualization/VirtualizationSystem.h
@@ -31,7 +31,7 @@ struct FPayloadActivityInfo
 };
 
 /** Describes the type of storage to use for a given action */
-enum EStorageType
+enum class EStorageType : int8
 {
 	/** Store in the local cache backends, this can be called from any thread */
 	Local = 0,
@@ -39,6 +39,18 @@ enum EStorageType
 	Persistent
 };
 
+enum class FPayloadStatus : int8
+{
+	/** The payload id was not value */
+	Invalid = -1,
+	/** The payload was not found in any backend for the given storage type */
+	NotFound = 0,
+	/** The payload was found in at least one backend but was not found in all backends available for the given storage type */
+	Partial,
+	/** The payload was found in all of the backends available for the given storage type */
+	FoundAll
+};
+
 /**
  * Creates the global IVirtualizationSystem if it has not already been set up. This can be called explicitly
  * during process start up but it will also be called by IVirtualizationSystem::Get if it detects that the
@@ -76,13 +88,11 @@ public:
 	 * Push a payload to the virtualization backends.
 	 *
 	 * @param	Id				The identifier of the payload being pushed.
-	 * @param	Payload			The payload itself in FCompressedBuffer form, it is
-	 *							assumed that if the buffer is to be compressed that
-	 *							it will have been done by the caller.
-	 * @param	StorageType		The type of storage to push the payload to, see EStorageType
-	 *							for details.
-	 * @param	PackageContext	Name of the owning package, which can be used to provide context
-	 *							about the payload.
+	 * @param	Payload			The payload itself in FCompressedBuffer form, it is assumed that if the buffer is to 
+	 *							be compressed that it will have been done by the caller.
+	 * @param	StorageType		The type of storage to push the payload to, @See EStorageType for details.
+	 * @param	PackageContext	Name of the owning package, which can be used to provide context about the payload.
+	 * 
 	 * @return	True if at least one backend now contains the payload, otherwise false.
 	 */
 	virtual bool PushData(const FPayloadId& Id, const FCompressedBuffer& Payload, EStorageType StorageType, const FPackagePath& PackageContext) = 0;
@@ -91,14 +101,27 @@ public:
 	 * Pull a payload from the virtualization backends.
 	 *
 	 * @param	Id The identifier of the payload being pulled.
-	 * @return	The payload in the form of a FCompressedBuffer. No decompression will
-	 *			be applied to the payload, it is up to the caller if they want to
-	 *			retain the payload in compressed or uncompressed format.
-	 *			If no backend contained the payload then an empty invalid FCompressedBuffer
-	 *			will be returned.
+	 * @return	The payload in the form of a FCompressedBuffer. No decompression will be applied to the payload, it is 
+	 *			up to the caller if they want to retain the payload in compressed or uncompressed format.  If no
+	 *			backend contained the payload then an empty invalid FCompressedBuffer will be returned.
 	 */
 	virtual FCompressedBuffer PullData(const FPayloadId& Id) = 0;
 
+	
+
+	/**
+	 * Query if a number of payloads exist or not in the given storage type. 
+	 * 
+	 * @param	Ids					One or more payload identifiers to test
+	 * @param	StorageType			The type of storage to push the payload to, @See EStorageType for details.
+	 * @param	OutStatuses [out]	An array containing the results for each payload. @See FPayloadStatus
+	 * 								If the operation succeeds the array will be resized to match the size of Ids. 
+	 * 
+	 * @return	True if the operation succeeded and the contents of OutStatuses is valid. False if errors were 
+	 * 			encountered in which case the contents of OutStatuses should be ignored.
+	 */
+	virtual bool DoPayloadsExist(TArrayView<const FPayloadId> Ids, EStorageType StorageType, TArray<FPayloadStatus>& OutStatuses) = 0;
+
 	using GetPayloadActivityInfoFuncRef = TFunctionRef<void(const FString& DebugName, const FString& ConfigName, const FPayloadActivityInfo& PayloadInfo)>;
 
 	/** Access profiling info relating to payload activity per backend. Stats will only be collected if ENABLE_COOK_STATS is enabled.*/
diff --git a/Engine/Source/Runtime/CoreUObject/Private/UObject/PackageTrailer.cpp b/Engine/Source/Runtime/CoreUObject/Private/UObject/PackageTrailer.cpp
--- a/Engine/Source/Runtime/CoreUObject/Private/UObject/PackageTrailer.cpp
+++ b/Engine/Source/Runtime/CoreUObject/Private/UObject/PackageTrailer.cpp
@@ -2,6 +2,7 @@
 
 #include "UObject/PackageTrailer.h"
 
+#include "Algo/Count.h"
 #include "CoreGlobals.h"
 #include "Misc/ConfigCacheIni.h"
 #include "Misc/PackagePath.h"
@@ -423,6 +424,31 @@ TArray<Virtualization::FPayloadId> FPackageTrailer::GetPayloads(EPayloadFilter T
 	return Identifiers;
 }
 
+int32 FPackageTrailer::GetNumPayloads(EPayloadFilter Type) const
+{
+	int32 Count = 0;
+
+	switch (Type)
+	{
+	case EPayloadFilter::All:
+		Count = Header.PayloadLookupTable.Num();
+		break;
+
+	case EPayloadFilter::Local:
+		Count = Algo::CountIf(Header.PayloadLookupTable, [](const Private::FLookupTableEntry& Entry) { return !Entry.IsVirtualized(); });
+		break;
+
+	case EPayloadFilter::Virtualized:
+		Count = Algo::CountIf(Header.PayloadLookupTable, [](const Private::FLookupTableEntry& Entry) { return Entry.IsVirtualized(); });
+		break;
+
+	default:
+		checkNoEntry();
+	}
+	
+	return Count;
+}
+
 bool FindPayloadsInPackageFile(const FPackagePath& PackagePath, EPayloadFilter Filter, TArray<Virtualization::FPayloadId>& OutPayloadIds)
 {
 	if (FPackageName::IsTextPackageExtension(PackagePath.GetHeaderExtension()))
diff --git a/Engine/Source/Runtime/CoreUObject/Public/UObject/PackageTrailer.h b/Engine/Source/Runtime/CoreUObject/Public/UObject/PackageTrailer.h
--- a/Engine/Source/Runtime/CoreUObject/Public/UObject/PackageTrailer.h
+++ b/Engine/Source/Runtime/CoreUObject/Public/UObject/PackageTrailer.h
@@ -239,6 +239,12 @@ public:
 	FPackageTrailer() = default;
 	~FPackageTrailer() = default;
 
+	FPackageTrailer(const FPackageTrailer& Other) = default;
+	FPackageTrailer& operator=(const FPackageTrailer & Other) = default;
+
+	FPackageTrailer(FPackageTrailer&& Other) = default;
+	FPackageTrailer& operator=(FPackageTrailer&& Other) = default;
+
 	/** 
 	 * Serializes the trailer from the given archive assuming that the seek position of the archive is already at the correct position
 	 * for the trailer.
@@ -282,9 +288,12 @@ public:
 	/** Returns the total size of the of the trailer on disk in bytes */
 	[[nodiscard]] int64 GetTrailerLength() const;
 
-	/** Returns an array of the payloads that match the given type. @See EPayloadType */
+	/** Returns an array of the payloads that match the given filter type. @See EPayloadType */
 	[[nodiscard]] TArray<Virtualization::FPayloadId> GetPayloads(EPayloadFilter Type) const;
 
+	/** Returns the number of payloads that the trailer owns that match the given filter type. @See EPayloadType */
+	[[nodiscard]] int32 GetNumPayloads(EPayloadFilter Type) const;
+
 	struct FHeader
 	{
 		/** Unique value used to identify the header */
