Cleaned up FVirtualizationManager to avoid duplicate code comments.

#rnx

- Removed the code documentation for the methods overriden from IVirtualizationSystem as the interface is already documented.
- Set FVirtualizationManager as final since it will not be derived from.

diff --git a/Engine/Source/Developer/Virtualization/Private/VirtualizationManager.h b/Engine/Source/Developer/Virtualization/Private/VirtualizationManager.h
--- a/Engine/Source/Developer/Virtualization/Private/VirtualizationManager.h
+++ b/Engine/Source/Developer/Virtualization/Private/VirtualizationManager.h
@@ -79,7 +79,7 @@ class IVirtualizationBackendFactory;
 
 /** This is used as a wrapper around the various potential back end implementations. 
 	The calling code shouldn't need to care about which back ends are actually in use. */
-class FVirtualizationManager : public IVirtualizationSystem
+class FVirtualizationManager final : public IVirtualizationSystem
 {
 public:
 	using FRegistedFactories = TMap<FName, IVirtualizationBackendFactory*>;
@@ -88,63 +88,28 @@ public:
 	FVirtualizationManager();
 	virtual ~FVirtualizationManager();
 
-	/** Poll to see if content virtualization is enabled or not. */
+private:
+	/* IVirtualizationSystem implementation */
+
 	virtual bool IsEnabled() const override;
 	
-	/** 
-	 * Push a payload to the virtualization backends.
-	 * 
-	 * @param	Id				The identifier of the payload being pushed.
-	 * @param	Payload			The payload itself in FCompressedBuffer form, it is 
-	 *							assumed that if the buffer is to be compressed that
-	 *							it will have been done by the caller.
-	 * @param	StorageType		The type of storage to push the payload to, see EStorageType
-	 *							for details.
-	 * @param	PackageContext	Name of the owning package, this will be used for filtering and
-	 *							to provide better logging messages. @see ShouldVirtualizePackage
-	 * @return	True if at least one backend now contains the payload, otherwise false.
-	 */
 	virtual bool PushData(const FPayloadId& Id, const FCompressedBuffer& Payload, EStorageType StorageType, const FPackagePath& PackageContext) override;
 
-	/** 
-	 * Pull a payload from the virtualization backends.
-	 *
-	 * @param	Id The identifier of the payload being pulled.
-	 * @return	The payload in the form of a FCompressedBuffer. No decompression will
-	 *			be applied to the payload, it is up to the caller if they want to 
-	 *			retain the payload in compressed or uncompressed format.
-	 *			If no backend contained the payload then an empty invalid FCompressedBuffer
-	 *			will be returned.
-	 */
 	virtual FCompressedBuffer PullData(const FPayloadId& Id) override;
 
-	/**
-	 * Query if a number of payloads exist or not in the given storage type.
-	 *
-	 * @param	Ids					One or more payload identifiers to test
-	 * @param	StorageType			The type of storage to push the payload to, @See EStorageType for details.
-	 * @param	OutStatuses [out]	An array containing the results for each payload. @See FPayloadStatus
-	 * 								If the operation succeeds the array will be resized to match the size of Ids.
-	 *
-	 * @return	True if the operation succeeded and the contents of OutStatuses is valid. False if errors were
-	 * 			encountered in which case the contents of OutStatuses should be ignored.
-	 */
 	virtual bool DoPayloadsExist(TArrayView<const FPayloadId> Ids, EStorageType StorageType, TArray<FPayloadStatus>& OutStatuses) override;
 
-	/** Access profiling info relating to accumulated payload activity for all backends. Stats will only be collected if ENABLE_COOK_STATS is enabled.*/
 	virtual FPayloadActivityInfo GetAccumualtedPayloadActivityInfo() const override;
 
-	/** Access profiling info relating to payload activity per backend. Stats will only be collected if ENABLE_COOK_STATS is enabled.*/
 	virtual void GetPayloadActivityInfo( GetPayloadActivityInfoFuncRef ) const override;
 
-	/** Get event delegate for system notifications */
 	virtual FOnNotification& GetNotificationEvent() override
 	{
 		return NotificationEvent;
 	}
 	
 private:
-	
+
 	void ApplySettingsFromConfigFiles(const FConfigFile& PlatformEngineIni);
 	void ApplySettingsFromCmdline();
 	
