Clean up the code organization for virtualization backends

#rnx

- Made all backends final, they should not be derived from.
- Added header files for the remaining backends that were cpp only. The headers are in the modules private directory so will not be used elsewhere and it makes it easier to see the interface.

diff --git a/Engine/Source/Developer/Virtualization/Private/VirtualizationDDCBackend.h b/Engine/Source/Developer/Virtualization/Private/VirtualizationDDCBackend.h
--- a/Engine/Source/Developer/Virtualization/Private/VirtualizationDDCBackend.h
+++ b/Engine/Source/Developer/Virtualization/Private/VirtualizationDDCBackend.h
@@ -2,9 +2,10 @@
 
 #pragma once 
 
-#include "DerivedDataCacheKey.h"
 #include "IVirtualizationBackend.h"
 
+#include "DerivedDataCacheKey.h"
+
 namespace UE::DerivedData { enum class ECachePolicy : uint32; }
 
 namespace UE::Virtualization
@@ -24,13 +25,14 @@ namespace UE::Virtualization
  * LocalStorage: When set to true, the payloads can be stored locally. [Default=true]
  * RemoteStorage: When set to true, the payloads can be stored remotely. [Default=true]
  */
-class FDDCBackend : public IVirtualizationBackend
+class FDDCBackend final : public IVirtualizationBackend
 {
 public:
 	explicit FDDCBackend(FStringView ConfigName, FStringView InDebugName);
 	virtual ~FDDCBackend() = default;
 
-protected:
+private:
+	/* IVirtualizationBackend implementation */
 
 	virtual bool Initialize(const FString& ConfigEntry) override;
 
@@ -39,11 +41,12 @@ protected:
 	virtual FCompressedBuffer PullData(const FPayloadId& Id) override;
 	
 	virtual bool DoesPayloadExist(const FPayloadId& Id) override;
-	
+
+private:
 	/** The bucket being used to group together the virtualized payloads in storage */
 	FString BucketName;
 
-	/** */
+	/** The FCacheBucket used with the DDC, cached to avoid recreating it for each request */
 	UE::DerivedData::FCacheBucket Bucket;
 
 	/** The policy to use when uploading or downloading data from the cache */
diff --git a/Engine/Source/Developer/Virtualization/Private/VirtualizationFileBackend.h b/Engine/Source/Developer/Virtualization/Private/VirtualizationFileBackend.h
--- a/Engine/Source/Developer/Virtualization/Private/VirtualizationFileBackend.h
+++ b/Engine/Source/Developer/Virtualization/Private/VirtualizationFileBackend.h
@@ -28,13 +28,14 @@ namespace UE::Virtualization
  *					Remember that the max length of time that the process can stall attempting to read
  *					a payload file is RetryCount * RetryWaitTime; [Default=100ms]
  */
-class FFileSystemBackend : public IVirtualizationBackend
+class FFileSystemBackend final : public IVirtualizationBackend
 {
 public:
-	FFileSystemBackend(FStringView ConfigName, FStringView DebugName);
+	explicit FFileSystemBackend(FStringView ConfigName, FStringView DebugName);
 	virtual ~FFileSystemBackend() = default;
 
-protected:
+private:
+	/* IVirtualizationBackend implementation */
 
 	virtual bool Initialize(const FString& ConfigEntry) override;
 
@@ -44,6 +45,8 @@ protected:
 
 	virtual bool DoesPayloadExist(const FPayloadId& Id) override;
 	
+private:
+
 	void CreateFilePath(const FPayloadId& PayloadId, FStringBuilderBase& OutPath);
 
 	TUniquePtr<FArchive> OpenFileForReading(const TCHAR* FilePath);
diff --git a/Engine/Source/Developer/Virtualization/Private/VirtualizationJupiterBackend.cpp b/Engine/Source/Developer/Virtualization/Private/VirtualizationJupiterBackend.cpp
--- a/Engine/Source/Developer/Virtualization/Private/VirtualizationJupiterBackend.cpp
+++ b/Engine/Source/Developer/Virtualization/Private/VirtualizationJupiterBackend.cpp
@@ -1,13 +1,13 @@
 // Copyright Epic Games, Inc. All Rights Reserved.
 
-#include "IVirtualizationBackend.h"
+#include "VirtualizationJupiterBackend.h"
 
 // TODO: Our libcurl implementation does not currently support MacOS
 // (not registering this will cause a fatal log error if the backend is actually used)
 #if PLATFORM_WINDOWS
-#define WITH_MIRAGE_JUPITER_BACKEND 1
+	#define WITH_MIRAGE_JUPITER_BACKEND 1
 #else
-#define WITH_MIRAGE_JUPITER_BACKEND 0
+	#define WITH_MIRAGE_JUPITER_BACKEND 0
 #endif //PLATFORM_WINDOWS
 
 #if WITH_MIRAGE_JUPITER_BACKEND
@@ -1294,771 +1294,718 @@ private:
 	FString SiteIdentifier;
 };
 
-/**
-* This backend allows data to be stored in and retrieved from a Jupiter service.
-* 
-* Ini file setup:
-* 'Name'=(Type=Jupiter, Host="", Namespace="", ChunkSize=, OAuthProvider="", OAuthClientId="", OAuthSecret="")
-* Host:				The URL of the service, use http://localhost if hosted locally.
-* Namespace:		Jupiter storage is divided into a number of namespaces allowing projects to keep their data separate
-*					while using the same service. This value controls which name space will be used.
-* ChunkSize:		Each payload can be divided into a number of chunks when being uploaded to Jupiter to improve upload 
-*					performance, this value sets the max size (in bytes) of each chunk. To disable and attempt to upload 
-*					each payload as a single data blob, set this to -1. 
-* OAuthProvider:	Url of the OAuth authorization server.
-* OAuthClientId:	Public identifier for use with the OAuth authorization server.
-* OAuthSecret:		Password for the OAuthClientId
-* (Note that the OAuth entries are not required if hosting locally)
-*/
-class FJupiterBackend : public IVirtualizationBackend
+FJupiterBackend::FJupiterBackend(FStringView ConfigName, FStringView InDebugName)
+	: IVirtualizationBackend(ConfigName, InDebugName, EOperations::Both)
+	, Namespace(TEXT("mirage"))
+	, Bucket(TEXT("default"))
+	, ChunkSize(-1)
+	, FailedLoginAttempts(0)		
+{		
+}
+
+bool FJupiterBackend::Initialize(const FString& ConfigEntry)
 {
-public:
-	FJupiterBackend(FStringView ConfigName, FStringView InDebugName)
-		: IVirtualizationBackend(ConfigName, InDebugName, EOperations::Both)
-		, Namespace(TEXT("mirage"))
-		, Bucket(TEXT("default"))
-		, ChunkSize(-1)
-		, FailedLoginAttempts(0)		
+	using namespace Utility;
+
+	// Some fields are required and will give fatal errors if not found!
+
+	if (!FParse::Value(*ConfigEntry, TEXT("Host="), HostAddress))
 	{
-		
+		UE_LOG(LogVirtualization, Fatal, TEXT("'Host=' not found in the config file"));
+		return false;
 	}
 
-private:
-
-	virtual bool Initialize(const FString& ConfigEntry)
+	if (!FParse::Value(*ConfigEntry, TEXT("Namespace="), Namespace))
 	{
-		using namespace Utility;
+		UE_LOG(LogVirtualization, Fatal, TEXT("'Namespace=' not found in the config file"));
+		return false;
+	}
 
-		// Some fields are required and will give fatal errors if not found!
+	if (FParse::Value(*ConfigEntry, TEXT("ChunkSize="), ChunkSize))
+	{
+		UE_LOG(LogVirtualization, Log, TEXT("ChunkSize set to '%" UINT64_FMT "' bytes"), ChunkSize);
+	}
+	else
+	{
+		UE_LOG(LogVirtualization, Log, TEXT("Payloads will not be chunked!"));
+	}
 
-		if (!FParse::Value(*ConfigEntry, TEXT("Host="), HostAddress))
+	// If we are connecting to a locally hosted Jupiter service then we do not need authorization
+	if (!IsUsingLocalHost())
+	{
+		if (!FParse::Value(*ConfigEntry, TEXT("OAuthProvider="), OAuthProvider))
 		{
-			UE_LOG(LogVirtualization, Fatal, TEXT("'Host=' not found in the config file"));
+			UE_LOG(LogVirtualization, Fatal, TEXT("'OAuthProvider=' not found in the config file"));
 			return false;
 		}
 
-		if (!FParse::Value(*ConfigEntry, TEXT("Namespace="), Namespace))
+		if (!FParse::Value(*ConfigEntry, TEXT("OAuthSecret="), OAuthSecret))
 		{
-			UE_LOG(LogVirtualization, Fatal, TEXT("'Namespace=' not found in the config file"));
+			UE_LOG(LogVirtualization, Fatal, TEXT("'OAuthSecret=' not found in the config file"));
 			return false;
 		}
 
-		if (FParse::Value(*ConfigEntry, TEXT("ChunkSize="), ChunkSize))
-		{
-			UE_LOG(LogVirtualization, Log, TEXT("ChunkSize set to '%" UINT64_FMT "' bytes"), ChunkSize);
-		}
-		else
-		{
-			UE_LOG(LogVirtualization, Log, TEXT("Payloads will not be chunked!"));
-		}
-
-		// If we are connecting to a locally hosted Jupiter service then we do not need authorization
-		if (!IsUsingLocalHost())
-		{
-			if (!FParse::Value(*ConfigEntry, TEXT("OAuthProvider="), OAuthProvider))
-			{
-				UE_LOG(LogVirtualization, Fatal, TEXT("'OAuthProvider=' not found in the config file"));
-					return false;
-			}
-
-			if (!FParse::Value(*ConfigEntry, TEXT("OAuthSecret="), OAuthSecret))
-			{
-				UE_LOG(LogVirtualization, Fatal, TEXT("'OAuthSecret=' not found in the config file"));
-				return false;
-			}
-
-			if (!FParse::Value(*ConfigEntry, TEXT("OAuthClientId="), OAuthClientId))
-			{
-				UE_LOG(LogVirtualization, Fatal, TEXT("'OAuthClientId=' not found in the config file"));
-				return false;
-			}
-		}
-
-		UE_LOG(LogVirtualization, Log, TEXT("Attempting to connect to a Jupiter service at '%s' with namespace '%s'"), *HostAddress, *Namespace);
-
-		if (!IsServiceReady())
+		if (!FParse::Value(*ConfigEntry, TEXT("OAuthClientId="), OAuthClientId))
 		{
+			UE_LOG(LogVirtualization, Fatal, TEXT("'OAuthClientId=' not found in the config file"));
 			return false;
 		}
+	}
 
-		if (!AcquireAccessToken())
-		{
-			return false;
-		}
+	UE_LOG(LogVirtualization, Log, TEXT("Attempting to connect to a Jupiter service at '%s' with namespace '%s'"), *HostAddress, *Namespace);
 
-		if (!ValidateServiceVersion())
-		{
-			return false;
-		}
-		
-		RequestPool = MakeUnique<FRequestPool>(*HostAddress, AccessToken.Get());
+	if (!IsServiceReady())
+	{
+		return false;
+	}
 
-		return true; 
+	if (!AcquireAccessToken())
+	{
+		return false;
 	}
 
-	virtual EPushResult PushData(const FPayloadId& Id, const FCompressedBuffer& CompressedPayload, const FPackagePath& PackageContext) override
+	if (!ValidateServiceVersion())
 	{
-		using namespace Utility;
+		return false;
+	}
+
+	RequestPool = MakeUnique<FRequestPool>(*HostAddress, AccessToken.Get());
 
-		TRACE_CPUPROFILER_EVENT_SCOPE(FJupiterBackend::PushData);
+	return true;
+}
+
+EPushResult FJupiterBackend::PushData(const FPayloadId& Id, const FCompressedBuffer& CompressedPayload, const FPackagePath& PackageContext)
+{
+	using namespace Utility;
+
+	TRACE_CPUPROFILER_EVENT_SCOPE(FJupiterBackend::PushData);
 
 #if UE_CHECK_FOR_EXISTING_PAYLOADS
-		if (DoesPayloadExist(Id))
-		{
-			UE_LOG(LogVirtualization, Verbose, TEXT("Jupiter already has a copy of the payload '%s'"), *Id.ToString());
-			return EPushResult::PayloadAlreadyExisted;
-		}
+	if (DoesPayloadExist(Id))
+	{
+		UE_LOG(LogVirtualization, Verbose, TEXT("Jupiter already has a copy of the payload '%s'"), *Id.ToString());
+		return EPushResult::PayloadAlreadyExisted;
+	}
 #endif // UE_CHECK_FOR_EXISTING_PAYLOADS
 
-		FEuropaDDCCachePUTRequest PUTRequest;
+	FEuropaDDCCachePUTRequest PUTRequest;
 
-		// TODO: This is a waste, we shouldn't need to flatten the buffer, but it makes it easier to work
-		// with the existing chunked code. The chunking code is most likely going to be removed before 
-		// this backend goes to production so it is not worth fixing up to work with FCompressedBuffer 
-		// properly.
-		FSharedBuffer FlattenedPayload = CompressedPayload.GetCompressed().ToShared();
+	// TODO: This is a waste, we shouldn't need to flatten the buffer, but it makes it easier to work
+	// with the existing chunked code. The chunking code is most likely going to be removed before 
+	// this backend goes to production so it is not worth fixing up to work with FCompressedBuffer 
+	// properly.
+	FSharedBuffer FlattenedPayload = CompressedPayload.GetCompressed().ToShared();
 
-		const int64 NumChunks = FMath::DivideAndRoundUp(FlattenedPayload.GetSize(), ChunkSize);
-		if (NumChunks > MAX_int32)
-		{
-			UE_LOG(LogVirtualization, Error, TEXT("Too many chunks (%d) are required for the payload '%s', try increasing the ChunkSize"), NumChunks, *Id.ToString());
-			return EPushResult::Failed;
-		}
+	const int64 NumChunks = FMath::DivideAndRoundUp(FlattenedPayload.GetSize(), ChunkSize);
+	if (NumChunks > MAX_int32)
+	{
+		UE_LOG(LogVirtualization, Error, TEXT("Too many chunks (%d) are required for the payload '%s', try increasing the ChunkSize"), NumChunks, *Id.ToString());
+		return EPushResult::Failed;
+	}
 
-		PUTRequest.ChunkHashes.SetNum((int32)NumChunks);
+	PUTRequest.ChunkHashes.SetNum((int32)NumChunks);
 
-		const uint8* DataPtr = (const uint8*)FlattenedPayload.GetData();
+	const uint8* DataPtr = (const uint8*)FlattenedPayload.GetData();
 
-		std::atomic<int32> NumFailedChunks(0);
-		FGraphEventArray Tasks;
-		Tasks.Reserve((int32)NumChunks);
+	std::atomic<int32> NumFailedChunks(0);
+	FGraphEventArray Tasks;
+	Tasks.Reserve((int32)NumChunks);
 
-		// Create and process the chunks that make up the payload
-		for (int32 Index = 0; Index < NumChunks; ++Index )
-		{
-			const int64 ChunkStart = Index * ChunkSize;
-			const uint64 BytesInChunk = FMath::Min(ChunkSize, FlattenedPayload.GetSize() - ChunkStart);
+	// Create and process the chunks that make up the payload
+	for (int32 Index = 0; Index < NumChunks; ++Index)
+	{
+		const int64 ChunkStart = Index * ChunkSize;
+		const uint64 BytesInChunk = FMath::Min(ChunkSize, FlattenedPayload.GetSize() - ChunkStart);
 
-			TArrayView<const uint8> ChunkData(&DataPtr[ChunkStart], (int32)BytesInChunk);
-			FString& ChunkHashString = PUTRequest.ChunkHashes[Index];
+		TArrayView<const uint8> ChunkData(&DataPtr[ChunkStart], (int32)BytesInChunk);
+		FString& ChunkHashString = PUTRequest.ChunkHashes[Index];
 
-			auto Job = [this, ChunkData, &Id, &ChunkHashString, &NumFailedChunks]()
+		auto Job = [this, ChunkData, &Id, &ChunkHashString, &NumFailedChunks]()
+		{
+			if (!PostChunk(ChunkData, Id, ChunkHashString))
 			{
-				if (!PostChunk(ChunkData, Id, ChunkHashString))
-				{
-					NumFailedChunks++;
-				}
-			};
+				NumFailedChunks++;
+			}
+		};
 
 #if UE_ENABLE_ASYNC_CHUNK_ACCESS
-			Tasks.Add(FFunctionGraphTask::CreateAndDispatchWhenReady(MoveTemp(Job)));
+		Tasks.Add(FFunctionGraphTask::CreateAndDispatchWhenReady(MoveTemp(Job)));
 #else
-			Job();
+		Job();
 #endif //UE_ENABLE_ASYNC_CHUNK_ACCESS
-		}
+	}
 
-		// There is some expensive work that we can do while we wait for the chunks to finish their upload
-		const FIoHash PayloadHash = FIoHash::HashBuffer(FlattenedPayload.GetData(), FlattenedPayload.GetSize());
+	// There is some expensive work that we can do while we wait for the chunks to finish their upload
+	const FIoHash PayloadHash = FIoHash::HashBuffer(FlattenedPayload.GetData(), FlattenedPayload.GetSize());
 
-		PUTRequest.PayloadHash = LexToString(PayloadHash);
-		PUTRequest.MetaData.PayloadLength = (int64)FlattenedPayload.GetSize();
-		PUTRequest.MetaData.ChunkLength = ChunkSize;
+	PUTRequest.PayloadHash = LexToString(PayloadHash);
+	PUTRequest.MetaData.PayloadLength = (int64)FlattenedPayload.GetSize();
+	PUTRequest.MetaData.ChunkLength = ChunkSize;
 
-		{
-			TRACE_CPUPROFILER_EVENT_SCOPE(FJupiterBackend::PullData::WaitOnChunks);
-			FTaskGraphInterface::Get().WaitUntilTasksComplete(Tasks);
-		}
+	{
+		TRACE_CPUPROFILER_EVENT_SCOPE(FJupiterBackend::PullData::WaitOnChunks);
+		FTaskGraphInterface::Get().WaitUntilTasksComplete(Tasks);
+	}
 
-		if (NumFailedChunks > 0)
-		{
-			UE_LOG(LogVirtualization, Error, TEXT("Failed to upload %d chunks for the payload '%s'."), NumFailedChunks.load(), *Id.ToString());
-			return EPushResult::Failed;
-		}
+	if (NumFailedChunks > 0)
+	{
+		UE_LOG(LogVirtualization, Error, TEXT("Failed to upload %d chunks for the payload '%s'."), NumFailedChunks.load(), *Id.ToString());
+		return EPushResult::Failed;
+	}
 
-		UE_LOG(LogVirtualization, Verbose, TEXT("Successfully uploaded all chunks for the payload '%s'"), *Id.ToString());	
+	UE_LOG(LogVirtualization, Verbose, TEXT("Successfully uploaded all chunks for the payload '%s'"), *Id.ToString());
 
-		// Note that the ddc end point is used by both ddc and mirage
-		TStringBuilder<256> Uri;
-		Uri.Appendf(TEXT("api/v1/c/ddc/%s/%s/%s"), *Namespace, *Bucket, *Id.ToString());
+	// Note that the ddc end point is used by both ddc and mirage
+	TStringBuilder<256> Uri;
+	Uri.Appendf(TEXT("api/v1/c/ddc/%s/%s/%s"), *Namespace, *Bucket, *Id.ToString());
 
-		// Retry request until we get an accepted response or exhaust allowed number of attempts.
-		uint32 Attempts = 0;
-		while (++Attempts <= UE_MIRAGE_MAX_ATTEMPTS)
+	// Retry request until we get an accepted response or exhaust allowed number of attempts.
+	uint32 Attempts = 0;
+	while (++Attempts <= UE_MIRAGE_MAX_ATTEMPTS)
+	{
+		FScopedRequestPtr Request(RequestPool.Get());
+		if (Request.IsValid())
 		{
-			FScopedRequestPtr Request(RequestPool.Get());
-			if (Request.IsValid())
-			{	
-				const FString OutputString = PUTRequest.ToJson(false);
-				
-				TStringConversion<TStringConvert<TCHAR, ANSICHAR>> ConvertedString = StringCast<ANSICHAR>(*OutputString);
-				TArrayView<const uint8> AnsiStringBuffer((uint8*)ConvertedString.Get(),ConvertedString.Length());
+			const FString OutputString = PUTRequest.ToJson(false);
 
-				FRequest::Result Result = Request->PerformBlockingUpload<FRequest::PutJson>(*Uri, AnsiStringBuffer);
-				const int64 ResponseCode = Request->GetResponseCode();
+			TStringConversion<TStringConvert<TCHAR, ANSICHAR>> ConvertedString = StringCast<ANSICHAR>(*OutputString);
+			TArrayView<const uint8> AnsiStringBuffer((uint8*)ConvertedString.Get(), ConvertedString.Length());
 
-				if (ResponseCode == 200)
-				{
-					UE_LOG(LogVirtualization, Verbose, TEXT("Successfully uploaded the description for the payload '%s'"), *Id.ToString());	
-					return EPushResult::Success;
-				}
+			FRequest::Result Result = Request->PerformBlockingUpload<FRequest::PutJson>(*Uri, AnsiStringBuffer);
+			const int64 ResponseCode = Request->GetResponseCode();
 
-				if (!ShouldRetryOnError(ResponseCode))
-				{
-					UE_LOG(LogVirtualization, Error, TEXT("Failed with error code '%d' to upload header infomation about payload '%s'"), ResponseCode, *Id.ToString());
-					return EPushResult::Failed;
-				}
+			if (ResponseCode == 200)
+			{
+				UE_LOG(LogVirtualization, Verbose, TEXT("Successfully uploaded the description for the payload '%s'"), *Id.ToString());
+				return EPushResult::Success;
 			}
-		}
 
-		UE_LOG(LogVirtualization, Error, TEXT("Failed  '%d' attempts to upload header infomation about payload '%s'"), UE_MIRAGE_MAX_ATTEMPTS, *Id.ToString());
-		return EPushResult::Failed;
+			if (!ShouldRetryOnError(ResponseCode))
+			{
+				UE_LOG(LogVirtualization, Error, TEXT("Failed with error code '%d' to upload header infomation about payload '%s'"), ResponseCode, *Id.ToString());
+				return EPushResult::Failed;
+			}
+		}
 	}
 
-	virtual FCompressedBuffer PullData(const FPayloadId& Id) override
-	{
-		using namespace Utility;
+	UE_LOG(LogVirtualization, Error, TEXT("Failed  '%d' attempts to upload header infomation about payload '%s'"), UE_MIRAGE_MAX_ATTEMPTS, *Id.ToString());
+	return EPushResult::Failed;
+}
 
-		TRACE_CPUPROFILER_EVENT_SCOPE(FJupiterBackend::PullData);
+FCompressedBuffer FJupiterBackend::PullData(const FPayloadId& Id)
+{
+	using namespace Utility;
+
+	TRACE_CPUPROFILER_EVENT_SCOPE(FJupiterBackend::PullData);
 
-		// First we need to get the description of the payload from Europa
+	// First we need to get the description of the payload from Europa
 
-		// Note that the ddc end point is used by both ddc and mirage
-		// fields=contentSha1 - Ask for the Sha1 hash of the fully reconstructed payload
-		// fields=blobIdentifiers - As for a list of the Sha1 hash ids for the payload chunks that we need to access from
-		// the Io service, the ids will be in the correct order.
-		// fields=metadata - Ask for the payload metadata which contains info we can use later for optimizations.
-		TStringBuilder<256> Uri;
-		Uri.Appendf(TEXT("api/v1/c/ddc/%s/%s/%s.json?fields=contentHash&fields=blobIdentifiers&fields=metadata"), *Namespace, *Bucket, *Id.ToString());
+	// Note that the ddc end point is used by both ddc and mirage
+	// fields=contentSha1 - Ask for the Sha1 hash of the fully reconstructed payload
+	// fields=blobIdentifiers - As for a list of the Sha1 hash ids for the payload chunks that we need to access from
+	// the Io service, the ids will be in the correct order.
+	// fields=metadata - Ask for the payload metadata which contains info we can use later for optimizations.
+	TStringBuilder<256> Uri;
+	Uri.Appendf(TEXT("api/v1/c/ddc/%s/%s/%s.json?fields=contentHash&fields=blobIdentifiers&fields=metadata"), *Namespace, *Bucket, *Id.ToString());
 
-		FDDCCacheGETResponse Response;
+	FDDCCacheGETResponse Response;
 
-		// Retry request until we get an accepted response or exhaust allowed number of attempts.
-		int64 ResponseCode = 0;
-		uint32 Attempts = 0;
-		while (ResponseCode != 200 && ++Attempts <= UE_MIRAGE_MAX_ATTEMPTS)
+	// Retry request until we get an accepted response or exhaust allowed number of attempts.
+	int64 ResponseCode = 0;
+	uint32 Attempts = 0;
+	while (ResponseCode != 200 && ++Attempts <= UE_MIRAGE_MAX_ATTEMPTS)
+	{
+		FScopedRequestPtr Request(RequestPool.Get());
+		if (Request.IsValid())
 		{
-			FScopedRequestPtr Request(RequestPool.Get());
-			if (Request.IsValid())
-			{
-				const FRequest::Result Result = Request->PerformBlockingDownload(*Uri);
-				ResponseCode = Request->GetResponseCode();
+			const FRequest::Result Result = Request->PerformBlockingDownload(*Uri);
+			ResponseCode = Request->GetResponseCode();
 
-				// Request was successful, make sure we got all the expected data.
-				if (FRequest::IsSuccessfulResponse(ResponseCode))
-				{
-					if (!Response.FromJson(Request->GetResponseAsJsonObject()))
-					{
-						UE_LOG(LogVirtualization, Error, TEXT("Failed to parser the header infomation about payload '%s'"), *Id.ToString());
-						return FCompressedBuffer();
-					}
-				}
-				else if (ResponseCode == 400)
-				{
-					// Response 400 indicates that the payload does not exist in Jupiter. Note that it is faster to just make the request
-					// and check for the response rather than call ::DoesPayloadExist prior to requesting the json header because this way 
-					// we will only make a single request if the payload exists or not.
-					UE_LOG(LogVirtualization, Verbose, TEXT("[%s] Does not contain the payload '%s'"), *GetDebugName(), *Id.ToString());
-					return FCompressedBuffer();
-				}
-				else if(!ShouldRetryOnError(ResponseCode))
+			// Request was successful, make sure we got all the expected data.
+			if (FRequest::IsSuccessfulResponse(ResponseCode))
+			{
+				if (!Response.FromJson(Request->GetResponseAsJsonObject()))
 				{
-					UE_LOG(LogVirtualization, Error, TEXT("Failed with error code '%d' to download header infomation about payload '%s'"), ResponseCode, *Id.ToString());
+					UE_LOG(LogVirtualization, Error, TEXT("Failed to parser the header infomation about payload '%s'"), *Id.ToString());
 					return FCompressedBuffer();
 				}
 			}
+			else if (ResponseCode == 400)
+			{
+				// Response 400 indicates that the payload does not exist in Jupiter. Note that it is faster to just make the request
+				// and check for the response rather than call ::DoesPayloadExist prior to requesting the json header because this way 
+				// we will only make a single request if the payload exists or not.
+				UE_LOG(LogVirtualization, Verbose, TEXT("[%s] Does not contain the payload '%s'"), *GetDebugName(), *Id.ToString());
+				return FCompressedBuffer();
+			}
+			else if (!ShouldRetryOnError(ResponseCode))
+			{
+				UE_LOG(LogVirtualization, Error, TEXT("Failed with error code '%d' to download header infomation about payload '%s'"), ResponseCode, *Id.ToString());
+				return FCompressedBuffer();
+			}
 		}
+	}
 
-		if (ResponseCode != 200)
-		{
-			UE_LOG(LogVirtualization, Error, TEXT("Failed '%d' attempts to download header infomation about payload (last error code '%d')  '%s'"), UE_MIRAGE_MAX_ATTEMPTS, ResponseCode, *Id.ToString());
-			return FCompressedBuffer();
-		}
+	if (ResponseCode != 200)
+	{
+		UE_LOG(LogVirtualization, Error, TEXT("Failed '%d' attempts to download header infomation about payload (last error code '%d')  '%s'"), UE_MIRAGE_MAX_ATTEMPTS, ResponseCode, *Id.ToString());
+		return FCompressedBuffer();
+	}
 
-		UE_LOG(LogVirtualization, Verbose, TEXT("Successfully downloaded a description for the payload '%s'"), *Id.ToString());	
+	UE_LOG(LogVirtualization, Verbose, TEXT("Successfully downloaded a description for the payload '%s'"), *Id.ToString());
 
-		// Now that we have the payload description we can start pulling the chunks from the Io service
-		// and reconstruct the final payload.
-		FUniqueBuffer Payload = FUniqueBuffer::Alloc(Response.MetaData.PayloadLength);
-		
-		int64 BytesLeft = Response.MetaData.PayloadLength;
+	// Now that we have the payload description we can start pulling the chunks from the Io service
+	// and reconstruct the final payload.
+	FUniqueBuffer Payload = FUniqueBuffer::Alloc(Response.MetaData.PayloadLength);
 
-		FGraphEventArray Tasks;
-		Tasks.Reserve(Response.ChunkHashes.Num());
+	int64 BytesLeft = Response.MetaData.PayloadLength;
 
-		std::atomic<int32> NumFailedChunks(0);
+	FGraphEventArray Tasks;
+	Tasks.Reserve(Response.ChunkHashes.Num());
 
-		bool bAllChunksPulled = true;
-		uint8* PayloadPtr = (uint8*)Payload.GetData();
+	std::atomic<int32> NumFailedChunks(0);
 
-		for (const FString& HashString : Response.ChunkHashes)
-		{
-			checkf(BytesLeft > 0, TEXT("Ran out of buffer space before all payload chunks were read!"));
+	bool bAllChunksPulled = true;
+	uint8* PayloadPtr = (uint8*)Payload.GetData();
+
+	for (const FString& HashString : Response.ChunkHashes)
+	{
+		checkf(BytesLeft > 0, TEXT("Ran out of buffer space before all payload chunks were read!"));
 
-			const int64 BytesToRead = FMath::Min(BytesLeft, Response.MetaData.ChunkLength);
+		const int64 BytesToRead = FMath::Min(BytesLeft, Response.MetaData.ChunkLength);
 
-			auto Job = [this, PayloadPtr, BytesToRead, &Id, &HashString, &NumFailedChunks]()
+		auto Job = [this, PayloadPtr, BytesToRead, &Id, &HashString, &NumFailedChunks]()
+		{
+			if (!PullChunk(HashString, Id, PayloadPtr, BytesToRead))
 			{
-				if (!PullChunk(HashString, Id, PayloadPtr, BytesToRead))
-				{
-					NumFailedChunks++;
-				}
-			};
+				NumFailedChunks++;
+			}
+		};
 
 #if UE_ENABLE_ASYNC_CHUNK_ACCESS
-			Tasks.Add(FFunctionGraphTask::CreateAndDispatchWhenReady(MoveTemp(Job)));
+		Tasks.Add(FFunctionGraphTask::CreateAndDispatchWhenReady(MoveTemp(Job)));
 #else
-			Job();
+		Job();
 #endif //UE_ENABLE_ASYNC_CHUNK_ACCESS
 
-			PayloadPtr += BytesToRead;
-			BytesLeft -= BytesToRead;
-		}
-
-		{
-			TRACE_CPUPROFILER_EVENT_SCOPE(FJupiterBackend::PullData::WaitOnChunks);
-			FTaskGraphInterface::Get().WaitUntilTasksComplete(Tasks);
-		}
-
-		if (NumFailedChunks == 0)
-		{
-			UE_LOG(LogVirtualization, Verbose, TEXT("Successfully downloaded all chunks for the payload '%s'"), *Id.ToString());	
-			return FCompressedBuffer::FromCompressed(Payload.MoveToShared());
-		}
-		else
-		{
-			UE_LOG(LogVirtualization, Error, TEXT("Failed to download %d chunks for the payload '%s'"), NumFailedChunks.load(), *Id.ToString());
-			return FCompressedBuffer();
-		}
+		PayloadPtr += BytesToRead;
+		BytesLeft -= BytesToRead;
 	}
 
-	bool IsUsingLocalHost() const
 	{
-		return HostAddress.StartsWith(TEXT("http://localhost"));
+		TRACE_CPUPROFILER_EVENT_SCOPE(FJupiterBackend::PullData::WaitOnChunks);
+		FTaskGraphInterface::Get().WaitUntilTasksComplete(Tasks);
 	}
 
-	// TODO: Pretty much the same code as in FHttpDerivedDataBackend, another candidate for code sharing
-	bool IsServiceReady()
+	if (NumFailedChunks == 0)
+	{
+		UE_LOG(LogVirtualization, Verbose, TEXT("Successfully downloaded all chunks for the payload '%s'"), *Id.ToString());	
+		return FCompressedBuffer::FromCompressed(Payload.MoveToShared());
+	}
+	else
 	{
-		TRACE_CPUPROFILER_EVENT_SCOPE(FJupiterBackend::IsServiceReady);
+		UE_LOG(LogVirtualization, Error, TEXT("Failed to download %d chunks for the payload '%s'"), NumFailedChunks.load(), *Id.ToString());
+		return FCompressedBuffer();
+	}
+}
 
-		using namespace Utility;
+bool FJupiterBackend::DoesPayloadExist(const FPayloadId& Id)
+{
+	using namespace Utility;
+
+	TRACE_CPUPROFILER_EVENT_SCOPE(FJupiterBackend::DoesPayloadExist);
 
-		FRequest Request(*HostAddress, nullptr, false);
-		FRequest::Result Result = Request.PerformBlockingDownload(TEXT("health/ready"), nullptr);
+	// Note that the ddc end point is used by both ddc and mirage
+	TStringBuilder<256> Uri;
+	Uri.Appendf(TEXT("api/v1/c/ddc/%s/%s/%s"), *Namespace, *Bucket, *Id.ToString());
 
-		if (Result == FRequest::Success && FRequest::IsSuccessfulResponse(Request.GetResponseCode()))
+	// Retry request until we get an accepted response or exhaust allowed number of attempts.
+	uint32 Attempts = 0;
+	while (++Attempts <= UE_MIRAGE_MAX_ATTEMPTS)
+	{
+		FScopedRequestPtr Request(RequestPool.Get());
+		if (Request.IsValid())
 		{
-			UE_LOG(LogVirtualization, Log, TEXT("Jupiter service status: '%s'."), *Request.GetResponseAsString());
-			return true;
+			const FRequest::Result Result = Request->PerformBlockingQuery<FRequest::Head>(*Uri);
+			const int64 ResponseCode = Request->GetResponseCode();
+
+			if (FRequest::IsSuccessfulResponse(ResponseCode))
+			{
+				return true;
+			}
+			else if (ResponseCode == 400)
+			{
+				return false;
+			}
+
+			if (!ShouldRetryOnError(ResponseCode))
+			{
+				return false;
+			}
 		}
-		else
-		{
-			UE_LOG(LogVirtualization, Error, TEXT("Unable to reach Jupiter service at '%s'. Status: %d . Response: '%s'"), *HostAddress, Request.GetResponseCode(), *Request.GetResponseAsString());
-			return false;
-		}	
 	}
 
+	return false;
+}
+
+bool FJupiterBackend::IsUsingLocalHost() const
+{
+	return HostAddress.StartsWith(TEXT("http://localhost"));
+}
+
+bool FJupiterBackend::IsServiceReady() const
+{
 	// TODO: Pretty much the same code as in FHttpDerivedDataBackend, another candidate for code sharing
-	bool AcquireAccessToken()
+	TRACE_CPUPROFILER_EVENT_SCOPE(FJupiterBackend::IsServiceReady);
+
+	using namespace Utility;
+
+	FRequest Request(*HostAddress, nullptr, false);
+	FRequest::Result Result = Request.PerformBlockingDownload(TEXT("health/ready"), nullptr);
+
+	if (Result == FRequest::Success && FRequest::IsSuccessfulResponse(Request.GetResponseCode()))
 	{
-		TRACE_CPUPROFILER_EVENT_SCOPE(FJupiterBackend::AcquireAccessToken);
+		UE_LOG(LogVirtualization, Log, TEXT("Jupiter service status: '%s'."), *Request.GetResponseAsString());
+		return true;
+	}
+	else
+	{
+		UE_LOG(LogVirtualization, Error, TEXT("Unable to reach Jupiter service at '%s'. Status: %d . Response: '%s'"), *HostAddress, Request.GetResponseCode(), *Request.GetResponseAsString());
+		return false;
+	}
+}
 
-		using namespace Utility;
+bool FJupiterBackend::AcquireAccessToken()
+{
+	// TODO: Pretty much the same code as in FHttpDerivedDataBackend, another candidate for code sharing
+	TRACE_CPUPROFILER_EVENT_SCOPE(FJupiterBackend::AcquireAccessToken);
 
-		if (IsUsingLocalHost())
-		{
-			UE_LOG(LogVirtualization, Log, TEXT("Connecting to a local host '%s', so skipping authorization"), *HostAddress);
-			return true;
-		}
+	using namespace Utility;
+
+	if (IsUsingLocalHost())
+	{
+		UE_LOG(LogVirtualization, Log, TEXT("Connecting to a local host '%s', so skipping authorization"), *HostAddress);
+		return true;
+	}
+
+	// Avoid spamming the this if the service is down
+	if (FailedLoginAttempts > UE_MIRAGE_MAX_FAILED_LOGIN_ATTEMPTS)
+	{
+		return false;
+	}
+
+	ensureMsgf(OAuthProvider.StartsWith(TEXT("http://")) || OAuthProvider.StartsWith(TEXT("https://")),
+		TEXT("The OAuth provider %s is not valid. Needs to be a fully qualified url."),
+		*OAuthProvider
+	);
+
+	// In case many requests wants to update the token at the same time
+	// get the current serial while we wait to take the CS.
+	const uint32 WantsToUpdateTokenSerial = AccessToken.IsValid() ? AccessToken->GetSerial() : 0u;
 
-		// Avoid spamming the this if the service is down
-		if (FailedLoginAttempts > UE_MIRAGE_MAX_FAILED_LOGIN_ATTEMPTS)
+	{
+		FScopeLock Lock(&AccessCs);
+
+		// Check if someone has beaten us to update the token, then it 
+		// should now be valid.
+		if (AccessToken.IsValid() && AccessToken->GetSerial() > WantsToUpdateTokenSerial)
 		{
-			return false;
+			return true;
 		}
 
-		ensureMsgf(OAuthProvider.StartsWith(TEXT("http://")) || OAuthProvider.StartsWith(TEXT("https://")),
-			TEXT("The OAuth provider %s is not valid. Needs to be a fully qualified url."),
-			*OAuthProvider
-		);
+		const uint32 SchemeEnd = OAuthProvider.Find(TEXT("://")) + 3;
+		const uint32 DomainEnd = OAuthProvider.Find(TEXT("/"), ESearchCase::CaseSensitive, ESearchDir::FromStart, SchemeEnd);
+		FString AuthDomain(DomainEnd, *OAuthProvider);
+		FString Uri(*OAuthProvider + DomainEnd + 1);
 
-		// In case many requests wants to update the token at the same time
-		// get the current serial while we wait to take the CS.
-		const uint32 WantsToUpdateTokenSerial = AccessToken.IsValid() ? AccessToken->GetSerial() : 0u;
+		FRequest Request(*AuthDomain, nullptr, false);
 
+		// If contents of the secret string is a file path, resolve and read form data.
+		if (OAuthSecret.StartsWith(TEXT("\\\\")))
 		{
-			FScopeLock Lock(&AccessCs);
-
-			// Check if someone has beaten us to update the token, then it 
-			// should now be valid.
-			if (AccessToken.IsValid() && AccessToken->GetSerial() > WantsToUpdateTokenSerial)
+			FString SecretFileContents;
+			if (FFileHelper::LoadFileToString(SecretFileContents, *OAuthSecret))
 			{
-				return true;
+				// Overwrite the filepath with the actual content.
+				OAuthSecret = SecretFileContents;
 			}
-
-			const uint32 SchemeEnd = OAuthProvider.Find(TEXT("://")) + 3;
-			const uint32 DomainEnd = OAuthProvider.Find(TEXT("/"), ESearchCase::CaseSensitive, ESearchDir::FromStart, SchemeEnd);
-			FString AuthDomain(DomainEnd, *OAuthProvider);
-			FString Uri(*OAuthProvider + DomainEnd + 1);
-
-			FRequest Request(*AuthDomain, nullptr, false);
-
-			// If contents of the secret string is a file path, resolve and read form data.
-			if (OAuthSecret.StartsWith(TEXT("\\\\")))
+			else
 			{
-				FString SecretFileContents;
-				if (FFileHelper::LoadFileToString(SecretFileContents, *OAuthSecret))
-				{
-					// Overwrite the filepath with the actual content.
-					OAuthSecret = SecretFileContents;
-				}
-				else
-				{
-					UE_LOG(LogVirtualization, Warning, TEXT("Failed to read OAuth form data file (%s)."), *OAuthSecret);
-					return false;
-				}
+				UE_LOG(LogVirtualization, Warning, TEXT("Failed to read OAuth form data file (%s)."), *OAuthSecret);
+				return false;
 			}
+		}
 
-			FString OAuthFormData = FString::Printf(
-				TEXT("client_id=%s&scope=cache_access&grant_type=client_credentials&client_secret=%s"),
-				*OAuthClientId,
-				*OAuthSecret
-			);
+		FString OAuthFormData = FString::Printf(
+			TEXT("client_id=%s&scope=cache_access&grant_type=client_credentials&client_secret=%s"),
+			*OAuthClientId,
+			*OAuthSecret
+		);
 
-			TArray<uint8> FormData;
-			auto OAuthFormDataUTF8 = FTCHARToUTF8(*OAuthFormData);
-			FormData.Append((uint8*)OAuthFormDataUTF8.Get(), OAuthFormDataUTF8.Length());
+		TArray<uint8> FormData;
+		auto OAuthFormDataUTF8 = FTCHARToUTF8(*OAuthFormData);
+		FormData.Append((uint8*)OAuthFormDataUTF8.Get(), OAuthFormDataUTF8.Length());
 
-			FRequest::Result Result = Request.PerformBlockingUpload<FRequest::Post>(*Uri, MakeArrayView(FormData));
+		FRequest::Result Result = Request.PerformBlockingUpload<FRequest::Post>(*Uri, MakeArrayView(FormData));
 
-			if (Result == FRequest::Success && Request.GetResponseCode() == 200)
+		if (Result == FRequest::Success && Request.GetResponseCode() == 200)
+		{
+			TSharedPtr<FJsonObject> ResponseObject = Request.GetResponseAsJsonObject();
+			if (ResponseObject)
 			{
-				TSharedPtr<FJsonObject> ResponseObject = Request.GetResponseAsJsonObject();
-				if (ResponseObject)
+				FString AccessTokenString;
+				int32 ExpiryTimeSeconds = 0;
+				int32 CurrentTimeSeconds = int32(FPlatformTime::ToSeconds(FPlatformTime::Cycles()));
+
+				if (ResponseObject->TryGetStringField(TEXT("access_token"), AccessTokenString) &&
+					ResponseObject->TryGetNumberField(TEXT("expires_in"), ExpiryTimeSeconds))
 				{
-					FString AccessTokenString;
-					int32 ExpiryTimeSeconds = 0;
-					int32 CurrentTimeSeconds = int32(FPlatformTime::ToSeconds(FPlatformTime::Cycles()));
+					if (!AccessToken)
+					{
+						AccessToken = MakeUnique<FAccessToken>();
+					}
+					AccessToken->SetHeader(*AccessTokenString);
+					UE_LOG(LogVirtualization, Log, TEXT("Logged in to HTTP DDC services. Expires in %d seconds."), ExpiryTimeSeconds);
 
-					if (ResponseObject->TryGetStringField(TEXT("access_token"), AccessTokenString) &&
-						ResponseObject->TryGetNumberField(TEXT("expires_in"), ExpiryTimeSeconds))
+					//Schedule a refresh of the token ahead of expiry time (this will not work in commandlets)
+					if (!IsRunningCommandlet())
 					{
-						if (!AccessToken)
-						{
-							AccessToken = MakeUnique<FAccessToken>();
-						}
-						AccessToken->SetHeader(*AccessTokenString);
-						UE_LOG(LogVirtualization, Log, TEXT("Logged in to HTTP DDC services. Expires in %d seconds."), ExpiryTimeSeconds);
-
-						//Schedule a refresh of the token ahead of expiry time (this will not work in commandlets)
-						if (!IsRunningCommandlet())
-						{
-							FTSTicker::GetCoreTicker().AddTicker(FTickerDelegate::CreateLambda(
-								[this](float DeltaTime)
-								{
-									this->AcquireAccessToken();
-									return false;
-								}
-							), ExpiryTimeSeconds - 20.0f);
-						}
-						// Reset failed login attempts, the service is indeed alive.
-						FailedLoginAttempts = 0;
-						return true;
+						FTSTicker::GetCoreTicker().AddTicker(FTickerDelegate::CreateLambda(
+							[this](float DeltaTime)
+							{
+								this->AcquireAccessToken();
+								return false;
+							}
+						), ExpiryTimeSeconds - 20.0f);
 					}
+					// Reset failed login attempts, the service is indeed alive.
+					FailedLoginAttempts = 0;
+					return true;
 				}
 			}
-			else
-			{
-				UE_LOG(LogVirtualization, Warning, TEXT("Failed to log in to HTTP services. Server responed with code %d."), Request.GetResponseCode());
-				FailedLoginAttempts++;
-			}
 		}
-		return false;
+		else
+		{
+			UE_LOG(LogVirtualization, Warning, TEXT("Failed to log in to HTTP services. Server responed with code %d."), Request.GetResponseCode());
+			FailedLoginAttempts++;
+		}
 	}
+	return false;
+}
 
-	/** 
-	 * Request the status of the service that we are connected to and make sure that it supports the 
-	 * feature set we need and meets our minimum version requirements.
-	 */
-	bool ValidateServiceVersion()
-	{
-		TRACE_CPUPROFILER_EVENT_SCOPE(FJupiterBackend::ValidateServiceVersion);
+bool FJupiterBackend::ValidateServiceVersion()
+{
+	TRACE_CPUPROFILER_EVENT_SCOPE(FJupiterBackend::ValidateServiceVersion);
 
-		using namespace Utility;
+	using namespace Utility;
 
-		FJupiterServiceStatus Response;
+	FJupiterServiceStatus Response;
 
-		int64 ResponseCode = 0;
-		uint32 Attempts = 0;
-		while (ResponseCode != 200 && ++Attempts <= UE_MIRAGE_MAX_ATTEMPTS)
-		{
-			// We create the FRequest ourselves since the request pool does not yet exist
-			FRequest Request(*HostAddress, AccessToken.Get(), true);
-			FRequest::Result Result = Request.PerformBlockingDownload(TEXT("api/v1/status"));
+	int64 ResponseCode = 0;
+	uint32 Attempts = 0;
+	while (ResponseCode != 200 && ++Attempts <= UE_MIRAGE_MAX_ATTEMPTS)
+	{
+		// We create the FRequest ourselves since the request pool does not yet exist
+		FRequest Request(*HostAddress, AccessToken.Get(), true);
+		FRequest::Result Result = Request.PerformBlockingDownload(TEXT("api/v1/status"));
 
-			ResponseCode = Request.GetResponseCode();
+		ResponseCode = Request.GetResponseCode();
 
-			// Request was successful, make sure we got all the expected data.
-			if (ResponseCode == 200)
-			{
-				if (!Response.FromJson(Request.GetResponseAsJsonObject()))
-				{
-					UE_LOG(LogVirtualization, Error, TEXT("The response to 'api/v1/status' GET did not contain valid data!"));
-					return false;
-				}
-			}
-			else if (!ShouldRetryOnError(ResponseCode))
+		// Request was successful, make sure we got all the expected data.
+		if (ResponseCode == 200)
+		{
+			if (!Response.FromJson(Request.GetResponseAsJsonObject()))
 			{
-				UE_LOG(LogVirtualization, Error, TEXT("Failed with error code '%d' to access the services status"), ResponseCode);
+				UE_LOG(LogVirtualization, Error, TEXT("The response to 'api/v1/status' GET did not contain valid data!"));
 				return false;
 			}
 		}
-
-		if (ResponseCode != 200)
-		{
-			UE_LOG(LogVirtualization, Error, TEXT("Failed '%d' attempts to access the services status (last error code '%d')"), UE_MIRAGE_MAX_ATTEMPTS, ResponseCode);
-			return false;
-		}
-
-		// Check version number
-		if (!Response.DoesHaveValidVersion(UE_MIRAGE_JUPITER_MIN_MAJOR_VER, UE_MIRAGE_JUPITER_MIN_MINOR_VER, UE_MIRAGE_JUPITER_MIN_PATCH_VER))
-		{
-			return false;
-		}
-
-		if (!Response.SupportsCapability(TEXT("ddc")))
+		else if (!ShouldRetryOnError(ResponseCode))
 		{
-			UE_LOG(LogVirtualization, Error, TEXT("Jupiter service does not support Europa (ddc) capability"));
+			UE_LOG(LogVirtualization, Error, TEXT("Failed with error code '%d' to access the services status"), ResponseCode);
 			return false;
 		}
+	}
 
-		Response.LogStatusInfo();
+	if (ResponseCode != 200)
+	{
+		UE_LOG(LogVirtualization, Error, TEXT("Failed '%d' attempts to access the services status (last error code '%d')"), UE_MIRAGE_MAX_ATTEMPTS, ResponseCode);
+		return false;
+	}
 
-		return true;
+	// Check version number
+	if (!Response.DoesHaveValidVersion(UE_MIRAGE_JUPITER_MIN_MAJOR_VER, UE_MIRAGE_JUPITER_MIN_MINOR_VER, UE_MIRAGE_JUPITER_MIN_PATCH_VER))
+	{
+		return false;
 	}
 
-	// TODO: Pretty much the same code as in FHttpDerivedDataBackend, another candidate for code sharing
-	bool ShouldRetryOnError(int64 ResponseCode)
+	if (!Response.SupportsCapability(TEXT("ddc")))
 	{
-		// Access token might have expired, request a new token and try again.
-		if (ResponseCode == 401 && AcquireAccessToken())
-		{
-			return true;
-		}
+		UE_LOG(LogVirtualization, Error, TEXT("Jupiter service does not support Europa (ddc) capability"));
+		return false;
+	}
 
-		// Too many requests, make a new attempt
-		if (ResponseCode == 429)
-		{
-			return true;
-		}
+	Response.LogStatusInfo();
 
-		// Gateway timeout, it will most likely work if we try again
-		if (ResponseCode == 504)
-		{
-			return true;
-		}
+	return true;
+}
+	
+bool FJupiterBackend::ShouldRetryOnError(int64 ResponseCode)
+{
+	// TODO: Pretty much the same code as in FHttpDerivedDataBackend, another candidate for code sharing
+	// 
+	// Access token might have expired, request a new token and try again.
+	if (ResponseCode == 401 && AcquireAccessToken())
+	{
+		return true;
+	}
 
-		return false;
+	// Too many requests, make a new attempt
+	if (ResponseCode == 429)
+	{
+		return true;
 	}
 
-	bool PostChunk(const TArrayView<const uint8>& ChunkData, const FPayloadId& PayloadId, FString& OutHashAsString)
+	// Gateway timeout, it will most likely work if we try again
+	if (ResponseCode == 504)
 	{
-		TRACE_CPUPROFILER_EVENT_SCOPE(FJupiterBackend::PostChunk);
+		return true;
+	}
 
-		using namespace Utility;
+	return false;
+}
 
-		const FIoHash ChunkHash = FIoHash::HashBuffer(ChunkData.GetData(), ChunkData.Num());
-		OutHashAsString = LexToString(ChunkHash);
+bool FJupiterBackend::PostChunk(const TArrayView<const uint8>& ChunkData, const FPayloadId& PayloadId, FString& OutHashAsString)
+{
+	TRACE_CPUPROFILER_EVENT_SCOPE(FJupiterBackend::PostChunk);
+
+	using namespace Utility;
+
+	const FIoHash ChunkHash = FIoHash::HashBuffer(ChunkData.GetData(), ChunkData.Num());
+	OutHashAsString = LexToString(ChunkHash);
 
 #if UE_CHECK_FOR_EXISTING_CHUNKS
-		if (DoesChunkExist(HashAsString))
-		{
-			UE_LOG(LogVirtualization, Verbose, TEXT("Jupiter already has a copy of the chunk '%s' for payload '%s'"), *HashAsString, *PayloadId.ToString());
-			return true;
-		}
+	if (DoesChunkExist(HashAsString))
+	{
+		UE_LOG(LogVirtualization, Verbose, TEXT("Jupiter already has a copy of the chunk '%s' for payload '%s'"), *HashAsString, *PayloadId.ToString());
+		return true;
+	}
 #endif // UE_CHECK_FOR_EXISTING_CHUNKS
 
-		FScopedRequestPtr Request(RequestPool.Get());
-		// TODO: Another candidate for code sharing
-		Request->SetHeader(TEXT("X-Jupiter-IoHash"), *OutHashAsString);		
+	FScopedRequestPtr Request(RequestPool.Get());
+	// TODO: Another candidate for code sharing
+	Request->SetHeader(TEXT("X-Jupiter-IoHash"), *OutHashAsString);
 
-		TStringBuilder<256> Uri;
-		Uri.Appendf(TEXT("api/v1/s/%s/%s"), *Namespace, *OutHashAsString);
+	TStringBuilder<256> Uri;
+	Uri.Appendf(TEXT("api/v1/s/%s/%s"), *Namespace, *OutHashAsString);
 
-		// Retry request until we get an accepted response or exhaust allowed number of attempts.
-		uint32 Attempts = 0;
-		while (++Attempts <= UE_MIRAGE_MAX_ATTEMPTS)
+	// Retry request until we get an accepted response or exhaust allowed number of attempts.
+	uint32 Attempts = 0;
+	while (++Attempts <= UE_MIRAGE_MAX_ATTEMPTS)
+	{
+		if (Request.IsValid())
 		{
-			if (Request.IsValid())
-			{
-				const FRequest::Result Result = Request->PerformBlockingUpload<FRequest::Put>(*Uri, ChunkData);
-				const int64 ResponseCode = Request->GetResponseCode();
+			const FRequest::Result Result = Request->PerformBlockingUpload<FRequest::Put>(*Uri, ChunkData);
+			const int64 ResponseCode = Request->GetResponseCode();
 
-				if (ResponseCode == 200)
-				{
-					UE_LOG(LogVirtualization, Verbose, TEXT("Successfully uploaded a chunk '%s'for payload '%s'"), *OutHashAsString, *PayloadId.ToString());	
-					return true;
-				}
+			if (ResponseCode == 200)
+			{
+				UE_LOG(LogVirtualization, Verbose, TEXT("Successfully uploaded a chunk '%s'for payload '%s'"), *OutHashAsString, *PayloadId.ToString());
+				return true;
+			}
 
-				if (!ShouldRetryOnError(ResponseCode))
-				{
-					return false;
-				}
+			if (!ShouldRetryOnError(ResponseCode))
+			{
+				return false;
 			}
 		}
-
-		return false;
 	}
 
-	bool PullChunk(const FString& Hash, const FPayloadId& PayloadId, uint8* DataPtr, int64 BufferSize)
-	{
-		TRACE_CPUPROFILER_EVENT_SCOPE(FJupiterBackend::PullChunk);
-
-		using namespace Utility;
-
-		FScopedRequestPtr Request(RequestPool.Get());
+	return false;
+}
 
-		TStringBuilder<256> Uri;
-		Uri.Appendf(TEXT("api/v1/s/%s/%s"), *Namespace, *Hash);
+bool FJupiterBackend::PullChunk(const FString& Hash, const FPayloadId& PayloadId, uint8* DataPtr, int64 BufferSize)
+{
+	TRACE_CPUPROFILER_EVENT_SCOPE(FJupiterBackend::PullChunk);
 
-		// Retry request until we get an accepted response or exhaust allowed number of attempts.
-		uint32 Attempts = 0;
-		while (++Attempts <= UE_MIRAGE_MAX_ATTEMPTS)
-		{
-			if (Request.IsValid())
-			{
-				const FRequest::Result Result = Request->PerformBlockingDownload(*Uri, FMutableMemoryView(DataPtr, BufferSize));
+	using namespace Utility;
 
-				if (Result != FRequest::Success)
-				{
-					UE_LOG(LogVirtualization, Error, TEXT("Attempting to GET a payload chunk '%s' for payload '%s' failed due to an internal Curl error"), *Hash, *PayloadId.ToString());	
-					return false;
-				}
+	FScopedRequestPtr Request(RequestPool.Get());
 
-				const int64 ResponseCode = Request->GetResponseCode();
+	TStringBuilder<256> Uri;
+	Uri.Appendf(TEXT("api/v1/s/%s/%s"), *Namespace, *Hash);
 
-				if (ResponseCode == 200)
-				{
-					UE_LOG(LogVirtualization, Verbose, TEXT("Successfully downloaded a payload chunk '%s' for payload '%s'"), *Hash, *PayloadId.ToString());	
-					return true;
-				}
+	// Retry request until we get an accepted response or exhaust allowed number of attempts.
+	uint32 Attempts = 0;
+	while (++Attempts <= UE_MIRAGE_MAX_ATTEMPTS)
+	{
+		if (Request.IsValid())
+		{
+			const FRequest::Result Result = Request->PerformBlockingDownload(*Uri, FMutableMemoryView(DataPtr, BufferSize));
 
-				if (!ShouldRetryOnError(ResponseCode))
-				{
-					UE_LOG(LogVirtualization, Error, TEXT("Attempting to GET a payload chunk '%s' for payload '%s' failed with http response: %" INT64_FMT), *Hash, *PayloadId.ToString(), ResponseCode);	
-					return false;
-				}
+			if (Result != FRequest::Success)
+			{
+				UE_LOG(LogVirtualization, Error, TEXT("Attempting to GET a payload chunk '%s' for payload '%s' failed due to an internal Curl error"), *Hash, *PayloadId.ToString());	
+				return false;
 			}
-		}
 
-		UE_LOG(LogVirtualization, Error, TEXT("Attempting to GET a payload chunk '%s' for payload '%s' failed all '%d' attempts"), *Hash, *PayloadId.ToString(), UE_MIRAGE_MAX_ATTEMPTS);	
-		return false;
-	}
+			const int64 ResponseCode = Request->GetResponseCode();
 
-	virtual bool DoesPayloadExist(const FPayloadId& Id) override
-	{
-		using namespace Utility;
-
-		TRACE_CPUPROFILER_EVENT_SCOPE(FJupiterBackend::DoesPayloadExist);
-
-		// Note that the ddc end point is used by both ddc and mirage
-		TStringBuilder<256> Uri;
-		Uri.Appendf(TEXT("api/v1/c/ddc/%s/%s/%s"), *Namespace, *Bucket, *Id.ToString());
-		
-		// Retry request until we get an accepted response or exhaust allowed number of attempts.
-		uint32 Attempts = 0;
-		while (++Attempts <= UE_MIRAGE_MAX_ATTEMPTS)
-		{
-			FScopedRequestPtr Request(RequestPool.Get());
-			if (Request.IsValid())
+			if (ResponseCode == 200)
 			{
-				const FRequest::Result Result = Request->PerformBlockingQuery<FRequest::Head>(*Uri);
-				const int64 ResponseCode = Request->GetResponseCode();
-
-				if (FRequest::IsSuccessfulResponse(ResponseCode))
-				{
-					return true;
-				}
-				else if (ResponseCode == 400)
-				{
-					return false;
-				}
+				UE_LOG(LogVirtualization, Verbose, TEXT("Successfully downloaded a payload chunk '%s' for payload '%s'"), *Hash, *PayloadId.ToString());	
+				return true;
+			}
 
-				if (!ShouldRetryOnError(ResponseCode))
-				{
-					return false;
-				}
+			if (!ShouldRetryOnError(ResponseCode))
+			{
+				UE_LOG(LogVirtualization, Error, TEXT("Attempting to GET a payload chunk '%s' for payload '%s' failed with http response: %" INT64_FMT), *Hash, *PayloadId.ToString(), ResponseCode);	
+				return false;
 			}
 		}
-
-		return false;
 	}
 
-	bool DoesChunkExist(const FString& Hash)
-	{
-		using namespace Utility;
+	UE_LOG(LogVirtualization, Error, TEXT("Attempting to GET a payload chunk '%s' for payload '%s' failed all '%d' attempts"), *Hash, *PayloadId.ToString(), UE_MIRAGE_MAX_ATTEMPTS);	
+	return false;
+}
 
-		TRACE_CPUPROFILER_EVENT_SCOPE(FJupiterBackend::DoesChunkExist);
+bool FJupiterBackend::DoesChunkExist(const FString& Hash)
+{
+	using namespace Utility;
+
+	TRACE_CPUPROFILER_EVENT_SCOPE(FJupiterBackend::DoesChunkExist);
 
-		TStringBuilder<256> Uri;
-		Uri.Appendf(TEXT("api/v1/s/%s/%s"), *Namespace, *Hash);
+	TStringBuilder<256> Uri;
+	Uri.Appendf(TEXT("api/v1/s/%s/%s"), *Namespace, *Hash);
 
-		// Retry request until we get an accepted response or exhaust allowed number of attempts.
-		uint32 Attempts = 0;
-		while (++Attempts <= UE_MIRAGE_MAX_ATTEMPTS)
+	// Retry request until we get an accepted response or exhaust allowed number of attempts.
+	uint32 Attempts = 0;
+	while (++Attempts <= UE_MIRAGE_MAX_ATTEMPTS)
+	{
+		FScopedRequestPtr Request(RequestPool.Get());
+		if (Request.IsValid())
 		{
-			FScopedRequestPtr Request(RequestPool.Get());
-			if (Request.IsValid())
-			{
-				const FRequest::Result Result = Request->PerformBlockingQuery<FRequest::Head>(*Uri);
-				const int64 ResponseCode = Request->GetResponseCode();
+			const FRequest::Result Result = Request->PerformBlockingQuery<FRequest::Head>(*Uri);
+			const int64 ResponseCode = Request->GetResponseCode();
 
-				if (FRequest::IsSuccessfulResponse(ResponseCode))
-				{
-					return true;
-				}
-				else if (ResponseCode == 404)
-				{
-					return false;
-				}
+			if (FRequest::IsSuccessfulResponse(ResponseCode))
+			{
+				return true;
+			}
+			else if (ResponseCode == 404)
+			{
+				return false;
+			}
 
-				if (!ShouldRetryOnError(ResponseCode))
-				{
-					return false;
-				}
+			if (!ShouldRetryOnError(ResponseCode))
+			{
+				return false;
 			}
 		}
-
-		return false;
 	}
-	
-	/** Address of the Jupiter service*/
-	FString HostAddress;
-	/** Namespace to connect to */
-	FString Namespace;
-	/** Europa allows us to organize the payloads by bucket. Currently this is not exposed and just set to 'default' */
-	FString Bucket;
-
-	/** The max size of each payload chunk */
-	uint64 ChunkSize;
-
-	/** Url of the OAuth authorization server */
-	FString OAuthProvider;
-	/**  Public identifier for use with the OAuth authorization server */
-	FString OAuthClientId;
-	/** Password for the OAuthClientId */
-	FString OAuthSecret;
-
-	/** The pool of FRequest objects that can be recycled */
-	TUniquePtr<Utility::FRequestPool> RequestPool;
-
-	/** Critical section used to protect the creation of new access tokens */
-	FCriticalSection AccessCs;
-	/** The access token used with service authorization */
-	TUniquePtr<Utility::FAccessToken> AccessToken;
-	/** Count how many times a login has failed since the last successful login */
-	uint32 FailedLoginAttempts;
-};
+
+	return false;
+}
 
 UE_REGISTER_VIRTUALIZATION_BACKEND_FACTORY(FJupiterBackend, Jupiter);
 
diff --git a/Engine/Source/Developer/Virtualization/Private/VirtualizationJupiterBackend.h b/Engine/Source/Developer/Virtualization/Private/VirtualizationJupiterBackend.h
new file mode 100644
--- /dev/null
+++ b/Engine/Source/Developer/Virtualization/Private/VirtualizationJupiterBackend.h
@@ -0,0 +1,91 @@
+// Copyright Epic Games, Inc. All Rights Reserved.
+
+#pragma once
+
+#include "IVirtualizationBackend.h"
+
+namespace UE::Utility { struct FRequestPool; }
+namespace UE::Utility { struct FAccessToken; }
+
+namespace UE::Virtualization
+{
+/**
+* This backend allows data to be stored in and retrieved from a Jupiter service.
+*
+* Ini file setup:
+* 'Name'=(Type=Jupiter, Host="", Namespace="", ChunkSize=, OAuthProvider="", OAuthClientId="", OAuthSecret="")
+* Host:				The URL of the service, use http://localhost if hosted locally.
+* Namespace:		Jupiter storage is divided into a number of namespaces allowing projects to keep their data separate
+*					while using the same service. This value controls which name space will be used.
+* ChunkSize:		Each payload can be divided into a number of chunks when being uploaded to Jupiter to improve upload
+*					performance, this value sets the max size (in bytes) of each chunk. To disable and attempt to upload
+*					each payload as a single data blob, set this to -1.
+* OAuthProvider:	Url of the OAuth authorization server.
+* OAuthClientId:	Public identifier for use with the OAuth authorization server.
+* OAuthSecret:		Password for the OAuthClientId
+* (Note that the OAuth entries are not required if hosting locally)
+*/
+class FJupiterBackend final : public IVirtualizationBackend
+{
+public:
+	explicit FJupiterBackend(FStringView ConfigName, FStringView DebugName);
+	virtual ~FJupiterBackend() = default;
+
+private:
+	/* IVirtualizationBackend implementation */
+
+	virtual bool Initialize(const FString& ConfigEntry) override;
+
+	virtual EPushResult PushData(const FPayloadId& Id, const FCompressedBuffer& CompressedPayload, const FPackagePath& PackageContext) override;
+
+	virtual FCompressedBuffer PullData(const FPayloadId& Id) override;
+
+	virtual bool DoesPayloadExist(const FPayloadId& Id) override;
+
+private:
+
+	bool IsUsingLocalHost() const;
+	bool IsServiceReady() const;
+	bool AcquireAccessToken();
+
+	/**
+	 * Request the status of the service that we are connected to and make sure that it supports the
+	 * feature set we need and meets our minimum version requirements.
+	 */
+	bool ValidateServiceVersion();
+
+	bool ShouldRetryOnError(int64 ResponseCode);
+
+	bool PostChunk(const TArrayView<const uint8>& ChunkData, const FPayloadId& PayloadId, FString& OutHashAsString);
+	bool PullChunk(const FString& Hash, const FPayloadId& PayloadId, uint8* DataPtr, int64 BufferSize);
+	bool DoesChunkExist(const FString& Hash);
+
+	/** Address of the Jupiter service*/
+	FString HostAddress;
+	/** Namespace to connect to */
+	FString Namespace;
+	/** Europa allows us to organize the payloads by bucket. Currently this is not exposed and just set to 'default' */
+	FString Bucket;
+
+	/** The max size of each payload chunk */
+	uint64 ChunkSize;
+
+	/** Url of the OAuth authorization server */
+	FString OAuthProvider;
+	/**  Public identifier for use with the OAuth authorization server */
+	FString OAuthClientId;
+	/** Password for the OAuthClientId */
+	FString OAuthSecret;
+
+	/** The pool of FRequest objects that can be recycled */
+	TUniquePtr<Utility::FRequestPool> RequestPool;
+
+	/** Critical section used to protect the creation of new access tokens */
+	FCriticalSection AccessCs;
+	/** The access token used with service authorization */
+	TUniquePtr<Utility::FAccessToken> AccessToken;
+	/** Count how many times a login has failed since the last successful login */
+	uint32 FailedLoginAttempts;
+};
+
+} // namespace UE::Virtualization
diff --git a/Engine/Source/Developer/Virtualization/Private/VirtualizationSourceControlBackend.cpp b/Engine/Source/Developer/Virtualization/Private/VirtualizationSourceControlBackend.cpp
--- a/Engine/Source/Developer/Virtualization/Private/VirtualizationSourceControlBackend.cpp
+++ b/Engine/Source/Developer/Virtualization/Private/VirtualizationSourceControlBackend.cpp
@@ -1,6 +1,6 @@
 // Copyright Epic Games, Inc. All Rights Reserved.
 
-#include "IVirtualizationBackend.h"
+#include "VirtualizationSourceControlBackend.h"
 
 #include "HAL/FileManager.h"
 #include "ISourceControlModule.h"
@@ -14,7 +14,6 @@
 #include "VirtualizationSourceControlUtilities.h"
 #include "VirtualizationUtilities.h"
 
-
 // When the SourceControl module (or at least the perforce source control module) is thread safe we
 // can enable this and stop using the hacky work around 'TryToDownloadFileFromBackgroundThread'
 #define IS_SOURCE_CONTROL_THREAD_SAFE 0
@@ -36,372 +35,351 @@ void CreateDescription(const FPackagePath& PackageContext, TStringBuilder<512>&
 	}
 }
 
-/**
- * This backend can be used to access payloads stored in source control.
- * The backend doesn't 'check out' a payload file but instead will just download the payload as
- * a binary blob.
- * It is assumed that the files are stored with the same path convention as the file system 
- * backend, found in Utils::PayloadIdToPath.
- * 
- * Ini file setup:
- * 'Name'=(Type=SourceControl, DepotRoot="//XXX/")
- * Where 'Name' is the backend name in the hierarchy and 'XXX' is the path in the source control
- * depot where the payload files are being stored.
- */
-class FSourceControlBackend : public IVirtualizationBackend
+FSourceControlBackend::FSourceControlBackend(FStringView ConfigName, FStringView InDebugName)
+	: IVirtualizationBackend(ConfigName, InDebugName, EOperations::Both)
 {
-public:
-	FSourceControlBackend(FStringView ConfigName, FStringView InDebugName)
-		: IVirtualizationBackend(ConfigName, InDebugName, EOperations::Both)
-	{
-	}
+}
 
-	virtual bool Initialize(const FString& ConfigEntry) override
-	{
-		TRACE_CPUPROFILER_EVENT_SCOPE(FSourceControlBackend::Initialize);
+bool FSourceControlBackend::Initialize(const FString& ConfigEntry)
+{
+	TRACE_CPUPROFILER_EVENT_SCOPE(FSourceControlBackend::Initialize);
 
-		// We require that a valid depot root has been provided
-		if (!FParse::Value(*ConfigEntry, TEXT("DepotRoot="), DepotRoot))
-		{
-			UE_LOG(LogVirtualization, Error, TEXT("'DepotRoot=' not found in the config file"));
-			return false;
-		}
+	// We require that a valid depot root has been provided
+	if (!FParse::Value(*ConfigEntry, TEXT("DepotRoot="), DepotRoot))
+	{
+		UE_LOG(LogVirtualization, Error, TEXT("'DepotRoot=' not found in the config file"));
+		return false;
+	}
 
-		ISourceControlModule& SSCModule = ISourceControlModule::Get();
+	ISourceControlModule& SSCModule = ISourceControlModule::Get();
 
-		// We require perforce as the source control provider as it is currently the only one that has the virtualization functionality implemented
-		const FName SourceControlName = SSCModule.GetProvider().GetName();
-		if (SourceControlName.IsNone())
-		{
-			// No source control provider is set so we can try to set it to "Perforce"
-			// Note this call will fatal error if "Perforce" is not a valid option
-			SSCModule.SetProvider(FName("Perforce"));
-		}
-		else if (SourceControlName != TEXT("Perforce"))
-		{
-			UE_LOG(LogVirtualization, Error, TEXT("Attempting to initialize FSourceControlBackend but source control is '%s' and only Perforce is currently supported!"), *SourceControlName.ToString());
-			return false;
-		}
+	// We require perforce as the source control provider as it is currently the only one that has the virtualization functionality implemented
+	const FName SourceControlName = SSCModule.GetProvider().GetName();
+	if (SourceControlName.IsNone())
+	{
+		// No source control provider is set so we can try to set it to "Perforce"
+		// Note this call will fatal error if "Perforce" is not a valid option
+		SSCModule.SetProvider(FName("Perforce"));
+	}
+	else if (SourceControlName != TEXT("Perforce"))
+	{
+		UE_LOG(LogVirtualization, Error, TEXT("Attempting to initialize FSourceControlBackend but source control is '%s' and only Perforce is currently supported!"), *SourceControlName.ToString());
+		return false;
+	}
 
-		ISourceControlProvider& SCCProvider = SSCModule.GetProvider();
+	ISourceControlProvider& SCCProvider = SSCModule.GetProvider();
 
-		if (!SCCProvider.IsAvailable())
-		{
-			SCCProvider.Init();
-		}
+	if (!SCCProvider.IsAvailable())
+	{
+		SCCProvider.Init();
+	}
 
-		// When a source control depot is set up a file named 'payload_metainfo.txt' should be submitted to it's root.
-		// This allows us to check for the existence of the file to confirm that the depot root is indeed valid.
-		const FString PayloadMetaInfoPath = FString::Printf(TEXT("%spayload_metainfo.txt"), *DepotRoot);
+	// When a source control depot is set up a file named 'payload_metainfo.txt' should be submitted to it's root.
+	// This allows us to check for the existence of the file to confirm that the depot root is indeed valid.
+	const FString PayloadMetaInfoPath = FString::Printf(TEXT("%spayload_metainfo.txt"), *DepotRoot);
 
 #if IS_SOURCE_CONTROL_THREAD_SAFE
-		TSharedRef<FDownloadFile, ESPMode::ThreadSafe> DownloadCommand = ISourceControlOperation::Create<FDownloadFile>();
-		if (SCCProvider.Execute(DownloadCommand, PayloadMetaInfoPath, EConcurrency::Synchronous) != ECommandResult::Succeeded)
-		{
-			UE_LOG(LogVirtualization, Error, TEXT("Failed to find 'payload_metainfo.txt' in the depot '%s', is your config set up correctly?"), *DepotRoot);
-			return false;
-		}	
+	TSharedRef<FDownloadFile, ESPMode::ThreadSafe> DownloadCommand = ISourceControlOperation::Create<FDownloadFile>();
+	if (SCCProvider.Execute(DownloadCommand, PayloadMetaInfoPath, EConcurrency::Synchronous) != ECommandResult::Succeeded)
+	{
+		UE_LOG(LogVirtualization, Error, TEXT("Failed to find 'payload_metainfo.txt' in the depot '%s', is your config set up correctly?"), *DepotRoot);
+		return false;
+	}	
 #else
-		TSharedRef<FDownloadFile, ESPMode::ThreadSafe> DownloadCommand = ISourceControlOperation::Create<FDownloadFile>();
-		if (!SCCProvider.TryToDownloadFileFromBackgroundThread(DownloadCommand, PayloadMetaInfoPath))
-		{
-			UE_LOG(LogVirtualization, Error, TEXT("Failed to find 'payload_metainfo.txt' in the depot '%s', is your config set up correctly?"), *DepotRoot);
-			return false;
-		}		
+	TSharedRef<FDownloadFile, ESPMode::ThreadSafe> DownloadCommand = ISourceControlOperation::Create<FDownloadFile>();
+	if (!SCCProvider.TryToDownloadFileFromBackgroundThread(DownloadCommand, PayloadMetaInfoPath))
+	{
+		UE_LOG(LogVirtualization, Error, TEXT("Failed to find 'payload_metainfo.txt' in the depot '%s', is your config set up correctly?"), *DepotRoot);
+		return false;
+	}		
 #endif //IS_SOURCE_CONTROL_THREAD_SAFE
 
-		FSharedBuffer MetaInfoBuffer = DownloadCommand->GetFileData(PayloadMetaInfoPath);
-		if (MetaInfoBuffer.IsNull())
-		{
-			UE_LOG(LogVirtualization, Error, TEXT("Failed to find 'payload_metainfo.txt' in the depot '%s', is your config set up correctly?"), *DepotRoot);
-			return false;
-		}
+	FSharedBuffer MetaInfoBuffer = DownloadCommand->GetFileData(PayloadMetaInfoPath);
+	if (MetaInfoBuffer.IsNull())
+	{
+		UE_LOG(LogVirtualization, Error, TEXT("Failed to find 'payload_metainfo.txt' in the depot '%s', is your config set up correctly?"), *DepotRoot);
+		return false;
+	}
 
-		// Currently we do not do anything with the payload meta info, in the future we could structure
-		// it's format to include more information that might be worth logging or something. 
-		// But for now being able to pull the payload meta info path at least shows that we can use the
-		// depot.
+	// Currently we do not do anything with the payload meta info, in the future we could structure
+	// it's format to include more information that might be worth logging or something. 
+	// But for now being able to pull the payload meta info path at least shows that we can use the
+	// depot.
 
-		return true;
-	}
+	return true;
+}
 
-	virtual EPushResult PushData(const FPayloadId& Id, const FCompressedBuffer& Payload, const FPackagePath& PackageContext) override
-	{
-		TRACE_CPUPROFILER_EVENT_SCOPE(FSourceControlBackend::PushData);
+EPushResult FSourceControlBackend::PushData(const FPayloadId& Id, const FCompressedBuffer& Payload, const FPackagePath& PackageContext)
+{
+	TRACE_CPUPROFILER_EVENT_SCOPE(FSourceControlBackend::PushData);
 
-		// TODO: Consider creating one workspace and one temp dir per session rather than per push.
-		// Although this would require more checking on start up to check for lingering workspaces
-		// and directories in case of editor crashes.
-		// We'd also need to remove each submitted file from the workspace after submission so that
-		// we can delete the local file
+	// TODO: Consider creating one workspace and one temp dir per session rather than per push.
+	// Although this would require more checking on start up to check for lingering workspaces
+	// and directories in case of editor crashes.
+	// We'd also need to remove each submitted file from the workspace after submission so that
+	// we can delete the local file
 
-		// We cannot easily submit files from within the project root due to p4 ignore rules
-		// so we will use the user temp directory instead. We append a guid to the root directory
-		// to avoid potentially conflicting with other editor processes that might be running.
+	// We cannot easily submit files from within the project root due to p4 ignore rules
+	// so we will use the user temp directory instead. We append a guid to the root directory
+	// to avoid potentially conflicting with other editor processes that might be running.
 
-		const FGuid SessionGuid = FGuid::NewGuid();
-		TStringBuilder<260> RootDirectory;
-		RootDirectory << FPlatformProcess::UserTempDir() << TEXT("UnrealEngine/VirtualizedPayloads/") << SessionGuid << TEXT("/");
+	const FGuid SessionGuid = FGuid::NewGuid();
+	TStringBuilder<260> RootDirectory;
+	RootDirectory << FPlatformProcess::UserTempDir() << TEXT("UnrealEngine/VirtualizedPayloads/") << SessionGuid << TEXT("/");
 
-		// First we need to save the payload to a file in the workspace client mapping so that it can be submitted
-		TStringBuilder<52> LocalPayloadPath;
-		Utils::PayloadIdToPath(Id, LocalPayloadPath);
+	// First we need to save the payload to a file in the workspace client mapping so that it can be submitted
+	TStringBuilder<52> LocalPayloadPath;
+	Utils::PayloadIdToPath(Id, LocalPayloadPath);
 
-		FString PayloadFilePath = *WriteToString<512>(RootDirectory, LocalPayloadPath);
+	FString PayloadFilePath = *WriteToString<512>(RootDirectory, LocalPayloadPath);
 
-		ON_SCOPE_EXIT
-		{
-			// Clean up the payload file from disk and the temp directories, but we do not need to give errors if any of these operations fail.
-			IFileManager::Get().Delete(*PayloadFilePath, false, false, true);
-			IFileManager::Get().DeleteDirectory(RootDirectory.ToString(), false, true);
-		};
+	ON_SCOPE_EXIT
+	{
+		// Clean up the payload file from disk and the temp directories, but we do not need to give errors if any of these operations fail.
+		IFileManager::Get().Delete(*PayloadFilePath, false, false, true);
+		IFileManager::Get().DeleteDirectory(RootDirectory.ToString(), false, true);
+	};
+
+	// Write the payload to Disk
+	{
+		UE_LOG(LogVirtualization, Verbose, TEXT("[%s] Writing payload to '%s' for submission"), *GetDebugName(), *PayloadFilePath);
 
-		// Write the payload to Disk
+		TUniquePtr<FArchive> FileAr(IFileManager::Get().CreateFileWriter(*PayloadFilePath));
+		if (!FileAr)
 		{
-			UE_LOG(LogVirtualization, Verbose, TEXT("[%s] Writing payload to '%s' for submission"), *GetDebugName(), *PayloadFilePath);
-
-			TUniquePtr<FArchive> FileAr(IFileManager::Get().CreateFileWriter(*PayloadFilePath));
-			if (!FileAr)
-			{
-				TStringBuilder<MAX_SPRINTF> SystemErrorMsg;
-				Utils::GetFormattedSystemError(SystemErrorMsg);
-
-				UE_LOG(	LogVirtualization, Error, TEXT("[%s] Failed to write payload '%s' contents to '%s' due to system error: %s"),
-						*GetDebugName(), 
-						*Id.ToString(), 
-						*PayloadFilePath, 
-						SystemErrorMsg.ToString());
-
-				return EPushResult::Failed;
-			}
-
-			*FileAr << const_cast<FCompressedBuffer&>(Payload);
-
-			if (!FileAr->Close())
-			{
-				TStringBuilder<MAX_SPRINTF> SystemErrorMsg;
-				Utils::GetFormattedSystemError(SystemErrorMsg);
-
-				UE_LOG(	LogVirtualization, Error, TEXT("[%s] Failed to write payload '%s' contents to '%s' due to system error: %s"),
-						*GetDebugName(), 
-					*Id.ToString(),
-					*PayloadFilePath,
-					SystemErrorMsg.
-					ToString());
-
-				return EPushResult::Failed;
-			}
-		}
+			TStringBuilder<MAX_SPRINTF> SystemErrorMsg;
+			Utils::GetFormattedSystemError(SystemErrorMsg);
+
+			UE_LOG(LogVirtualization, Error, TEXT("[%s] Failed to write payload '%s' contents to '%s' due to system error: %s"),
+				*GetDebugName(),
+				*Id.ToString(),
+				*PayloadFilePath,
+				SystemErrorMsg.ToString());
 
-		TStringBuilder<64> WorkspaceName;
-		WorkspaceName << TEXT("MirageSubmission-") << SessionGuid;
+			return EPushResult::Failed;
+		}
 
-		ISourceControlProvider& SCCProvider = ISourceControlModule::Get().GetProvider();
+		*FileAr << const_cast<FCompressedBuffer&>(Payload);
 
-		// Create a temp workspace so that we can submit the payload from
+		if (!FileAr->Close())
 		{
-			TSharedRef<FCreateWorkspace> CreateWorkspaceCommand = ISourceControlOperation::Create<FCreateWorkspace>(WorkspaceName, RootDirectory);
+			TStringBuilder<MAX_SPRINTF> SystemErrorMsg;
+			Utils::GetFormattedSystemError(SystemErrorMsg);
 
-			TStringBuilder<512> DepotMapping;
-			DepotMapping << DepotRoot << TEXT("...");
+			UE_LOG(LogVirtualization, Error, TEXT("[%s] Failed to write payload '%s' contents to '%s' due to system error: %s"),
+				*GetDebugName(),
+				*Id.ToString(),
+				*PayloadFilePath,
+				SystemErrorMsg.
+				ToString());
 
-			TStringBuilder<128> ClientMapping;
-			ClientMapping << TEXT("//") << WorkspaceName << TEXT("/...");
+			return EPushResult::Failed;
+		}
+	}
 
-			CreateWorkspaceCommand->AddNativeClientViewMapping(DepotMapping, ClientMapping);
+	TStringBuilder<64> WorkspaceName;
+	WorkspaceName << TEXT("MirageSubmission-") << SessionGuid;
 
-			if (SCCProvider.Execute(CreateWorkspaceCommand) != ECommandResult::Succeeded)
-			{
-				UE_LOG(	LogVirtualization, Error, TEXT("[%s] Failed to create temp workspace '%s' to submit payload '%s' from"), 
-						*GetDebugName(), 
-						WorkspaceName.ToString(), 
-						*Id.ToString());
+	ISourceControlProvider& SCCProvider = ISourceControlModule::Get().GetProvider();
 
-				return EPushResult::Failed;
-			}
-		}
+	// Create a temp workspace so that we can submit the payload from
+	{
+		TSharedRef<FCreateWorkspace> CreateWorkspaceCommand = ISourceControlOperation::Create<FCreateWorkspace>(WorkspaceName, RootDirectory);
 
-		ON_SCOPE_EXIT
-		{
-			// Remove the temp workspace mapping
-			if (SCCProvider.Execute(ISourceControlOperation::Create<FDeleteWorkspace>(WorkspaceName)) != ECommandResult::Succeeded)
-			{
-				UE_LOG(LogVirtualization, Warning, TEXT("[%s] Failed to remove temp workspace '%s' please delete manually"), *GetDebugName(), WorkspaceName.ToString());
-			}
-		};
-
-		FSourceControlResultInfo SwitchToNewWorkspaceInfo;
-		FString OriginalWorkspace;
-		if (SCCProvider.SwitchWorkspace(WorkspaceName, SwitchToNewWorkspaceInfo, &OriginalWorkspace) != ECommandResult::Succeeded)
-		{
-			UE_LOG(	LogVirtualization, Error, TEXT("[%s] Failed to switch to temp workspace '%s' when trying to submit payload '%s'"), 
-					*GetDebugName(), 
-					WorkspaceName.ToString(), 
-					*Id.ToString());
+		TStringBuilder<512> DepotMapping;
+		DepotMapping << DepotRoot << TEXT("...");
 
-			return EPushResult::Failed;
-		}
+		TStringBuilder<128> ClientMapping;
+		ClientMapping << TEXT("//") << WorkspaceName << TEXT("/...");
 
-		ON_SCOPE_EXIT
-		{
-			FSourceControlResultInfo SwitchToOldWorkspaceInfo;
-			if (SCCProvider.SwitchWorkspace(OriginalWorkspace, SwitchToOldWorkspaceInfo, nullptr) != ECommandResult::Succeeded)
-			{
-				// Failing to restore the old workspace could result in confusing editor issues and data loss, so for now it is fatal.
-				// The medium term plan should be to refactor the SourceControlModule so that we could use an entirely different 
-				// ISourceControlProvider so as not to affect the rest of the editor.
-				UE_LOG(	LogVirtualization, Fatal, TEXT("[%s] Failed to restore the original workspace to temp workspace '%s' continuing would risk editor instability and potential data loss"), 
-						*GetDebugName(), 
-						*OriginalWorkspace);
-			}
-		};
-
-		FSourceControlStatePtr FileState = SCCProvider.GetState(PayloadFilePath, EStateCacheUsage::ForceUpdate);
-		if (!FileState.IsValid())
+		CreateWorkspaceCommand->AddNativeClientViewMapping(DepotMapping, ClientMapping);
+
+		if (SCCProvider.Execute(CreateWorkspaceCommand) != ECommandResult::Succeeded)
 		{
-			UE_LOG(LogVirtualization, Error, TEXT("[%s] Failed to find the current file state for '%s'"), *GetDebugName(), *PayloadFilePath);
+			UE_LOG(LogVirtualization, Error, TEXT("[%s] Failed to create temp workspace '%s' to submit payload '%s' from"),
+				*GetDebugName(),
+				WorkspaceName.ToString(),
+				*Id.ToString());
+
 			return EPushResult::Failed;
 		}
+	}
 
-		if (FileState->IsSourceControlled())
+	ON_SCOPE_EXIT
+	{
+		// Remove the temp workspace mapping
+		if (SCCProvider.Execute(ISourceControlOperation::Create<FDeleteWorkspace>(WorkspaceName)) != ECommandResult::Succeeded)
 		{
-			// TODO: Maybe check if the data is the same (could be different if the compression algorithm has changed)
-			// TODO: Should we respect if the file is deleted as technically we can still get access to it?
-			return EPushResult::PayloadAlreadyExisted;
+			UE_LOG(LogVirtualization, Warning, TEXT("[%s] Failed to remove temp workspace '%s' please delete manually"), *GetDebugName(), WorkspaceName.ToString());
 		}
-		else if (FileState->CanAdd())
+	};
+
+	FSourceControlResultInfo SwitchToNewWorkspaceInfo;
+	FString OriginalWorkspace;
+	if (SCCProvider.SwitchWorkspace(WorkspaceName, SwitchToNewWorkspaceInfo, &OriginalWorkspace) != ECommandResult::Succeeded)
+	{
+		UE_LOG(LogVirtualization, Error, TEXT("[%s] Failed to switch to temp workspace '%s' when trying to submit payload '%s'"),
+			*GetDebugName(),
+			WorkspaceName.ToString(),
+			*Id.ToString());
+
+		return EPushResult::Failed;
+	}
+
+	ON_SCOPE_EXIT
+	{
+		FSourceControlResultInfo SwitchToOldWorkspaceInfo;
+		if (SCCProvider.SwitchWorkspace(OriginalWorkspace, SwitchToOldWorkspaceInfo, nullptr) != ECommandResult::Succeeded)
 		{
-			if (SCCProvider.Execute(ISourceControlOperation::Create<FMarkForAdd>(), PayloadFilePath) != ECommandResult::Succeeded)
-			{
-				UE_LOG(LogVirtualization, Error, TEXT("[%s] Failed to mark the payload file '%s' for Add in source control"), *GetDebugName(), *PayloadFilePath);
-				return EPushResult::Failed;
-			}
+			// Failing to restore the old workspace could result in confusing editor issues and data loss, so for now it is fatal.
+			// The medium term plan should be to refactor the SourceControlModule so that we could use an entirely different 
+			// ISourceControlProvider so as not to affect the rest of the editor.
+			UE_LOG(LogVirtualization, Fatal, TEXT("[%s] Failed to restore the original workspace to temp workspace '%s' continuing would risk editor instability and potential data loss"),
+					*GetDebugName(),
+					*OriginalWorkspace);
 		}
-		else
+	};
+
+	FSourceControlStatePtr FileState = SCCProvider.GetState(PayloadFilePath, EStateCacheUsage::ForceUpdate);
+	if (!FileState.IsValid())
+	{
+		UE_LOG(LogVirtualization, Error, TEXT("[%s] Failed to find the current file state for '%s'"), *GetDebugName(), *PayloadFilePath);
+		return EPushResult::Failed;
+	}
+
+	if (FileState->IsSourceControlled())
+	{
+		// TODO: Maybe check if the data is the same (could be different if the compression algorithm has changed)
+		// TODO: Should we respect if the file is deleted as technically we can still get access to it?
+		return EPushResult::PayloadAlreadyExisted;
+	}
+	else if (FileState->CanAdd())
+	{
+		if (SCCProvider.Execute(ISourceControlOperation::Create<FMarkForAdd>(), PayloadFilePath) != ECommandResult::Succeeded)
 		{
-			UE_LOG(LogVirtualization, Error, TEXT("[%s] The the payload file '%s' is not in source control but also cannot be marked for Add"), *GetDebugName(), *PayloadFilePath);
+			UE_LOG(LogVirtualization, Error, TEXT("[%s] Failed to mark the payload file '%s' for Add in source control"), *GetDebugName(), *PayloadFilePath);
 			return EPushResult::Failed;
 		}
+	}
+	else
+	{
+		UE_LOG(LogVirtualization, Error, TEXT("[%s] The the payload file '%s' is not in source control but also cannot be marked for Add"), *GetDebugName(), *PayloadFilePath);
+		return EPushResult::Failed;
+	}
 
-		// Now submit the payload
-		{
-			TSharedRef<FCheckIn, ESPMode::ThreadSafe> CheckInOperation = ISourceControlOperation::Create<FCheckIn>();
+	// Now submit the payload
+	{
+		TSharedRef<FCheckIn, ESPMode::ThreadSafe> CheckInOperation = ISourceControlOperation::Create<FCheckIn>();
 
-			TStringBuilder<512> Description;
-			CreateDescription(PackageContext, Description);
+		TStringBuilder<512> Description;
+		CreateDescription(PackageContext, Description);
 
-			CheckInOperation->SetDescription(FText::FromString(Description.ToString()));
+		CheckInOperation->SetDescription(FText::FromString(Description.ToString()));
 
-			if (SCCProvider.Execute(CheckInOperation, PayloadFilePath) != ECommandResult::Succeeded)
-			{
-				UE_LOG(LogVirtualization, Error, TEXT("[%s] Failed to submit the payload file '%s' to source control"), *GetDebugName(), *PayloadFilePath);
-				return EPushResult::Failed;
-			}
+		if (SCCProvider.Execute(CheckInOperation, PayloadFilePath) != ECommandResult::Succeeded)
+		{
+			UE_LOG(LogVirtualization, Error, TEXT("[%s] Failed to submit the payload file '%s' to source control"), *GetDebugName(), *PayloadFilePath);
+			return EPushResult::Failed;
 		}
-
-		return EPushResult::Success;
 	}
 
-	virtual FCompressedBuffer PullData(const FPayloadId& Id) override
-	{
-		TRACE_CPUPROFILER_EVENT_SCOPE(FSourceControlBackend::PullData);
+	return EPushResult::Success;
+}
+
+FCompressedBuffer FSourceControlBackend::PullData(const FPayloadId& Id)
+{
+	TRACE_CPUPROFILER_EVENT_SCOPE(FSourceControlBackend::PullData);
 
-		TStringBuilder<512> DepotPath;
-		CreateDepotPath(Id, DepotPath);
+	TStringBuilder<512> DepotPath;
+	CreateDepotPath(Id, DepotPath);
 
-		ISourceControlProvider& SCCProvider = ISourceControlModule::Get().GetProvider();
+	ISourceControlProvider& SCCProvider = ISourceControlModule::Get().GetProvider();
 
 #if IS_SOURCE_CONTROL_THREAD_SAFE
-		TSharedRef<FDownloadFile, ESPMode::ThreadSafe> DownloadCommand = ISourceControlOperation::Create<FDownloadFile>(FDownloadFile::EVerbosity::None);
-		if (SCCProvider.Execute(DownloadCommand, DepotPath.ToString(), EConcurrency::Synchronous) != ECommandResult::Succeeded)
-		{
-			return FCompressedBuffer();
-		}
+	TSharedRef<FDownloadFile, ESPMode::ThreadSafe> DownloadCommand = ISourceControlOperation::Create<FDownloadFile>(FDownloadFile::EVerbosity::None);
+	if (SCCProvider.Execute(DownloadCommand, DepotPath.ToString(), EConcurrency::Synchronous) != ECommandResult::Succeeded)
+	{
+		return FCompressedBuffer();
+	}
 #else
-		TSharedRef<FDownloadFile> DownloadCommand = ISourceControlOperation::Create<FDownloadFile>(FDownloadFile::EVerbosity::None);
-		if (!SCCProvider.TryToDownloadFileFromBackgroundThread(DownloadCommand, DepotPath.ToString()))
-		{
-			return FCompressedBuffer();
-		}
+	TSharedRef<FDownloadFile> DownloadCommand = ISourceControlOperation::Create<FDownloadFile>(FDownloadFile::EVerbosity::None);
+	if (!SCCProvider.TryToDownloadFileFromBackgroundThread(DownloadCommand, DepotPath.ToString()))
+	{
+		return FCompressedBuffer();
+	}
 #endif
 
-		// The payload was created by FCompressedBuffer::Compress so we can return it 
-		// as a FCompressedBuffer.
-		FSharedBuffer Buffer = DownloadCommand->GetFileData(DepotPath);
-		return FCompressedBuffer::FromCompressed(Buffer);
-	}
+	// The payload was created by FCompressedBuffer::Compress so we can return it 
+	// as a FCompressedBuffer.
+	FSharedBuffer Buffer = DownloadCommand->GetFileData(DepotPath);
+	return FCompressedBuffer::FromCompressed(Buffer);
+}
 
-	void CreateDepotPath(const FPayloadId& PayloadId, FStringBuilderBase& OutPath)
-	{
-		TStringBuilder<52> PayloadPath;
-		Utils::PayloadIdToPath(PayloadId, PayloadPath);
+bool FSourceControlBackend::DoesPayloadExist(const FPayloadId& Id)
+{
+	TArray<bool> Result;
 
-		OutPath << DepotRoot << PayloadPath;
+	if (FSourceControlBackend::DoPayloadsExist(MakeArrayView<const FPayloadId>(&Id, 1), Result))
+	{
+		check(Result.Num() == 1);
+		return Result[0];
 	}
-
-	virtual bool DoesPayloadExist(const FPayloadId& Id)
+	else
 	{
-		TArray<bool> Result;
-
-		if (FSourceControlBackend::DoPayloadsExist(MakeArrayView<const FPayloadId>(&Id, 1), Result))
-		{
-			check(Result.Num() == 1);
-			return Result[0];
-		}
-		else
-		{
-			return false;
-		}
+		return false;
 	}
+}
 
-	virtual bool DoPayloadsExist(TArrayView<const FPayloadId> PayloadIds, TArray<bool>& OutResults) override
-	{
-		ISourceControlProvider& SCCProvider = ISourceControlModule::Get().GetProvider();
+bool FSourceControlBackend::DoPayloadsExist(TArrayView<const FPayloadId> PayloadIds, TArray<bool>& OutResults)
+{
+	ISourceControlProvider& SCCProvider = ISourceControlModule::Get().GetProvider();
 
-		TArray<FString> DepotPaths;
-		DepotPaths.Reserve(PayloadIds.Num());
+	TArray<FString> DepotPaths;
+	DepotPaths.Reserve(PayloadIds.Num());
 
-		TArray<FSourceControlStateRef> PathStates;
-	
-		for (const FPayloadId& PayloadId : PayloadIds)
-		{
-			if (PayloadId.IsValid())
-			{
-				TStringBuilder<52> LocalPayloadPath;
-				Utils::PayloadIdToPath(PayloadId, LocalPayloadPath);
+	TArray<FSourceControlStateRef> PathStates;
 
-				DepotPaths.Emplace(WriteToString<512>(DepotRoot, LocalPayloadPath));
-			}		
-		}
-	
-		ECommandResult::Type Result = SCCProvider.GetState(DepotPaths, PathStates, EStateCacheUsage::ForceUpdate);
-		if (Result != ECommandResult::Type::Succeeded)
+	for (const FPayloadId& PayloadId : PayloadIds)
+	{
+		if (PayloadId.IsValid())
 		{
-			UE_LOG(LogVirtualization, Error, TEXT("[%s] Failed to query the state of files in the source control depot"), *GetDebugName());
-			return false;
+			TStringBuilder<52> LocalPayloadPath;
+			Utils::PayloadIdToPath(PayloadId, LocalPayloadPath);
+
+			DepotPaths.Emplace(WriteToString<512>(DepotRoot, LocalPayloadPath));
 		}
+	}
+
+	ECommandResult::Type Result = SCCProvider.GetState(DepotPaths, PathStates, EStateCacheUsage::ForceUpdate);
+	if (Result != ECommandResult::Type::Succeeded)
+	{
+		UE_LOG(LogVirtualization, Error, TEXT("[%s] Failed to query the state of files in the source control depot"), *GetDebugName());
+		return false;
+	}
 
-		check(DepotPaths.Num() == PathStates.Num()); // We expect that all paths return a state
+	check(DepotPaths.Num() == PathStates.Num()); // We expect that all paths return a state
 
-		OutResults.SetNum(PayloadIds.Num());
+	OutResults.SetNum(PayloadIds.Num());
 
-		int32 StatusIndex = 0;
-		for (int32 Index = 0; Index < PayloadIds.Num(); ++Index)
+	int32 StatusIndex = 0;
+	for (int32 Index = 0; Index < PayloadIds.Num(); ++Index)
+	{
+		if (PayloadIds[Index].IsValid())
 		{
-			if (PayloadIds[Index].IsValid())
-			{
-				OutResults[Index] = PathStates[StatusIndex++]->IsSourceControlled();
-			}
+			OutResults[Index] = PathStates[StatusIndex++]->IsSourceControlled();
 		}
-
-		return true;
 	}
 
-private:
+	return true;
+}
 
-	/** The root where the virtualized payloads are stored in source control */
-	FString DepotRoot;
-};
+void FSourceControlBackend::CreateDepotPath(const FPayloadId& PayloadId, FStringBuilderBase& OutPath)
+{
+	TStringBuilder<52> PayloadPath;
+	Utils::PayloadIdToPath(PayloadId, PayloadPath);
+
+	OutPath << DepotRoot << PayloadPath;
+}
 
 UE_REGISTER_VIRTUALIZATION_BACKEND_FACTORY(FSourceControlBackend, SourceControl);
 
diff --git a/Engine/Source/Developer/Virtualization/Private/VirtualizationSourceControlBackend.h b/Engine/Source/Developer/Virtualization/Private/VirtualizationSourceControlBackend.h
new file mode 100644
--- /dev/null
+++ b/Engine/Source/Developer/Virtualization/Private/VirtualizationSourceControlBackend.h
@@ -0,0 +1,51 @@
+// Copyright Epic Games, Inc. All Rights Reserved.
+
+#pragma once 
+
+#include "IVirtualizationBackend.h"
+
+#include "Containers/StringView.h"
+
+namespace UE::Virtualization
+{
+
+/**
+ * This backend can be used to access payloads stored in source control.
+ * The backend doesn't 'check out' a payload file but instead will just download the payload as
+ * a binary blob.
+ * It is assumed that the files are stored with the same path convention as the file system
+ * backend, found in Utils::PayloadIdToPath.
+ *
+ * Ini file setup:
+ * 'Name'=(Type=SourceControl, DepotRoot="//XXX/")
+ * Where 'Name' is the backend name in the hierarchy and 'XXX' is the path in the source control
+ * depot where the payload files are being stored.
+ */
+class FSourceControlBackend final : public IVirtualizationBackend
+{
+public:
+	explicit FSourceControlBackend(FStringView ConfigName, FStringView InDebugName);
+	virtual ~FSourceControlBackend() = default;
+	
+private:
+	/* IVirtualizationBackend implementation */
+
+	virtual bool Initialize(const FString& ConfigEntry) override;
+	
+	virtual EPushResult PushData(const FPayloadId& Id, const FCompressedBuffer& Payload, const FPackagePath& PackageContext) override;
+	
+	virtual FCompressedBuffer PullData(const FPayloadId& Id) override;
+
+	virtual bool DoesPayloadExist(const FPayloadId& Id) override;
+	
+	virtual bool DoPayloadsExist(TArrayView<const FPayloadId> PayloadIds, TArray<bool>& OutResults) override;
+
+private:
+
+	void CreateDepotPath(const FPayloadId& PayloadId, FStringBuilderBase& OutPath);
+
+	/** The root where the virtualized payloads are stored in source control */
+	FString DepotRoot;
+};
+
+} // namespace UE::Virtualization
