CIS - Fix mac non unity compile

#rnx

- We need to include "Virtualization/PayloadId.h" due to the use of PayloadId in DoPayloadsExist.

diff --git a/Engine/Source/Developer/Virtualization/Public/IVirtualizationBackend.h b/Engine/Source/Developer/Virtualization/Public/IVirtualizationBackend.h
--- a/Engine/Source/Developer/Virtualization/Public/IVirtualizationBackend.h
+++ b/Engine/Source/Developer/Virtualization/Public/IVirtualizationBackend.h
@@ -7,14 +7,12 @@
 #include "Features/IModularFeature.h"
 #include "Features/IModularFeatures.h"
 #include "Templates/UniquePtr.h"
-
-class FPackagePath;
+#include "Virtualization/PayloadId.h"
+#include "Virtualization/VirtualizationSystem.h"
 
 namespace UE::Virtualization
 {
 
-class FPayloadId;
-
 /** Describes the result of a IVirtualizationBackend::Push operation */
 enum class EPushResult
 {
@@ -86,7 +84,21 @@ public:
 	 * @param Payload	A potentially compressed buffer representing the payload
 	 * @return			The result of the push operation
 	 */
-	virtual EPushResult PushData(const FPayloadId& Id, const FCompressedBuffer& Payload, const FPackagePath& PackageContext) = 0;
+	virtual EPushResult PushData(const FPayloadId& Id, const FCompressedBuffer& Payload, const FString& PackageContext) = 0;
+
+	virtual EPushResult PushData(TArrayView<FPushRequest> Requests)
+	{
+		// TODO: Sort return codes
+		for (const FPushRequest& Request : Requests)
+		{
+			if (PushData(Request.Identifier, Request.Payload, Request.Context) == EPushResult::Failed)
+			{
+				return EPushResult::Failed;
+			}
+		}
+
+		return EPushResult::Success;
+	}
 
 	/** 
 	 * The backend will attempt to retrieve the given payload by what ever method the backend uses.
