Allow the sourcecontrol virtualization backend to choose the type of the temp workspace used for payload submission.

#rnx

-  The source control backend will now attempt to create the temporary workspace as a partitioned workspace by default.
- This can be overriden via the config file when defining the virtualization graph.

diff --git a/Engine/Source/Developer/Virtualization/Private/VirtualizationSourceControlBackend.cpp b/Engine/Source/Developer/Virtualization/Private/VirtualizationSourceControlBackend.cpp
--- a/Engine/Source/Developer/Virtualization/Private/VirtualizationSourceControlBackend.cpp
+++ b/Engine/Source/Developer/Virtualization/Private/VirtualizationSourceControlBackend.cpp
@@ -51,6 +51,10 @@ bool FSourceControlBackend::Initialize(const FString& ConfigEntry)
 		return false;
 	}
 
+	// Optional config values
+	FParse::Bool(*ConfigEntry, TEXT("UsePartitionedClient="), bUsePartitionedClient);
+	UE_LOG(LogVirtualization, Log, TEXT("[%s] Using partitioned clients: '%s'"), *GetDebugName(), bUsePartitionedClient ? TEXT("true") : TEXT("false"));
+
 	ISourceControlModule& SSCModule = ISourceControlModule::Get();
 
 	// We require perforce as the source control provider as it is currently the only one that has the virtualization functionality implemented
@@ -194,6 +198,11 @@ EPushResult FSourceControlBackend::PushData(const FPayloadId& Id, const FCompres
 
 		CreateWorkspaceCommand->AddNativeClientViewMapping(DepotMapping, ClientMapping);
 
+		if (bUsePartitionedClient)
+		{
+			CreateWorkspaceCommand->SetType(FCreateWorkspace::EType::Partitioned);
+		}
+
 		if (SCCProvider.Execute(CreateWorkspaceCommand) != ECommandResult::Succeeded)
 		{
 			UE_LOG(LogVirtualization, Error, TEXT("[%s] Failed to create temp workspace '%s' to submit payload '%s' from"),
diff --git a/Engine/Source/Developer/Virtualization/Private/VirtualizationSourceControlBackend.h b/Engine/Source/Developer/Virtualization/Private/VirtualizationSourceControlBackend.h
--- a/Engine/Source/Developer/Virtualization/Private/VirtualizationSourceControlBackend.h
+++ b/Engine/Source/Developer/Virtualization/Private/VirtualizationSourceControlBackend.h
@@ -20,6 +20,12 @@ namespace UE::Virtualization
  * 'Name'=(Type=SourceControl, DepotRoot="//XXX/")
  * Where 'Name' is the backend name in the hierarchy and 'XXX' is the path in the source control
  * depot where the payload files are being stored.
+ * 
+ * Optional Values:
+ * UsePartitionedClient:	When true the temporary workspace client created to submit payloads 
+ *							from will be created as a partitioned workspace which is less overhead
+ *							on the source control server. If your server does not support this then
+  *							use false. [Default=True]
  */
 class FSourceControlBackend final : public IVirtualizationBackend
 {
@@ -46,6 +52,10 @@ private:
 
 	/** The root where the virtualized payloads are stored in source control */
 	FString DepotRoot;
+
+	/** Should we try to make the temp client partitioned or not? */
+	bool bUsePartitionedClient = true;
+
 };
 
 } // namespace UE::Virtualization
