Packages submitted from the editor together will now virtualize their payloads in a single batch rather than one at a time.

#rnx

### VirtualizationSystem
- Added a new overload for Push to VirtualizationSystem that takes an array of FPushRequest, which is a new structure representing a single payload request.
- Filtering by package name is currently disabled, this is because the API has been forced into changing and passing the package name in via a FString rather than FPackagePath which means we would need to be more careful. This will be done in a future submit.
- The backend interface has been extended to also have a batch version of PushData, by default this will attempt to submit each request one at a time so payloads don't have to try and implement a batched version if there is no need.
- The context being passed with a payload when being pushed has been changed from FPackagePath to FString due to include order issues, as the FPackagePath lives in CoreUObject and the API for virtualization lives in Core. Additionally in the future the payloads might not be owned by a package (there is nothing specifically enforcing this) so the context being a string makes more sense.
- NOTE: Due to the context change we currently no longer support the filtering feature, which allows for payloads belonging to packages under specific directories to be excluded from virtualization. This is something that will be solved in a future submit.

### SourceControlBackend

- Now that we can submit multiple payloads in the same submit, the CL description has been changed slightly. We will now print a list of payload identifiers -> the package trying to submit that payload. This will only tell the users which package originally caused the payload to submit. If a user submits a new package at a later date that contains the same payload we will not be updating the description.

### PackageSubmissionChecks
- Converted the submission process to use the new batch push operation in VirtualizationSystem.
-- This means that we do a single push and then have to update the package trailers to convert the now pushed payloads from local to virtualized.
- Added new define UE_PRECHECK_PAYLOAD_STATUS that makes it easy to toggle off the checks to see which payloads need to be submitted to the persistent backend.  This is useful to test if it actually helps speed up the overall operations or if it is faster to just perform the batch push operations on all payloads and check the return values.
-- The hope is that over time the submission processes will become fast enough that we can remove the precheck.
- Fixed up logging to not always assume more than one package or payload.

### General Notes
- Errors and logging is now a bit more vague as we often not just report that X payloads failed etc rather than specific payload identifiers. This probably doesn't affect the user too much since those identifiers as fairly meaningless to them anyway.
- The source control submission could be further optimized by first checking the status of the files in thge depot and only then creating/switching workspace etc.
- As currently written, we need to load all of the payloads into memory, then the backends will do what they need (in the case of source control this results in the payloads being written to disk then submitted) which could create quite a large memory spike when submitting a large number of packages.
-- One solution would be to change the batch push API to take a "payload provider" interface and have the payloads requested as needed rather than passing in the FCompressedBuffer directly. This would let us immediately write the payload to disk for submission then discard it from memory, preventing larger spikes. Although it could cause overhead if there are multiple backends being submitted to. Internally we are unlikely to have more than one backend per storage solution so maybe we should just make it a config option?

diff --git a/Engine/Source/Developer/Virtualization/Private/PackageSubmissionChecks.cpp b/Engine/Source/Developer/Virtualization/Private/PackageSubmissionChecks.cpp
--- a/Engine/Source/Developer/Virtualization/Private/PackageSubmissionChecks.cpp
+++ b/Engine/Source/Developer/Virtualization/Private/PackageSubmissionChecks.cpp
@@ -24,6 +24,10 @@
 // removed.
 #define UE_VALIDATE_TRUNCATED_PACKAGE 1
 
+// When enabled we will check the payloads to see if they already exist in the persistent storage
+// backends before trying to push them.
+#define UE_PRECHECK_PAYLOAD_STATUS 1
+
 namespace UE::Virtualization
 {
 
@@ -154,7 +158,7 @@ void OnPrePackageSubmission(const TArray<FString>& FilesToSubmit, TArray<FText>&
 		bool bWasTrailerUpdated = false;
 	};
 
-	UE_LOG(LogVirtualization, Display, TEXT("Considering %d files for virtualization"), FilesToSubmit.Num());
+	UE_LOG(LogVirtualization, Display, TEXT("Considering %d file(s) for virtualization"), FilesToSubmit.Num());
 
 	TArray<FPackageInfo> Packages;
 	Packages.Reserve(FilesToSubmit.Num());
@@ -194,7 +198,7 @@ void OnPrePackageSubmission(const TArray<FString>& FilesToSubmit, TArray<FText>&
 		}
 	}
 
-	UE_LOG(LogVirtualization, Display, TEXT("Found %" INT64_FMT " payloads in %d packages that need to be checked for virtualization"), TotalPayloadsToCheck, Packages.Num());
+	UE_LOG(LogVirtualization, Display, TEXT("Found %" INT64_FMT " payload(s) in %d package(s) that need to be examined for virtualization"), TotalPayloadsToCheck, Packages.Num());
 
 	TArray<FPayloadStatus> PayloadStatuses;
 	if (!System.DoPayloadsExist(AllLocalPayloads, EStorageType::Persistent, PayloadStatuses))
@@ -211,6 +215,7 @@ void OnPrePackageSubmission(const TArray<FString>& FilesToSubmit, TArray<FText>&
 	{
 		check(PackageInfo.LocalPayloads.IsEmpty() || PackageInfo.PayloadIndex != INDEX_NONE); // If we have payloads we should have an index
 
+#if UE_PRECHECK_PAYLOAD_STATUS
 		for (int32 Index = 0; Index < PackageInfo.LocalPayloads.Num(); ++Index)
 		{
 			if (PayloadStatuses[PackageInfo.PayloadIndex + Index] == FPayloadStatus::FoundAll)
@@ -220,7 +225,7 @@ void OnPrePackageSubmission(const TArray<FString>& FilesToSubmit, TArray<FText>&
 					PackageInfo.bWasTrailerUpdated = true;
 				}
 				else
-					{
+				{
 					FText Message = FText::Format(	LOCTEXT("Virtualization_UpdateStatusFailed", "Unable to update the status for the payload '{0}' in the package '{1}'"),
 													FText::FromString(PackageInfo.LocalPayloads[Index].ToString()),
 													FText::FromString(PackageInfo.Path.GetDebugName()));
@@ -235,13 +240,18 @@ void OnPrePackageSubmission(const TArray<FString>& FilesToSubmit, TArray<FText>&
 		{
 			PackageInfo.LocalPayloads = PackageInfo.Trailer.GetPayloads(EPayloadFilter::Local);
 		}
+#endif
 
+		PackageInfo.PayloadIndex = INDEX_NONE;
 		TotalPayloadsToVirtualize += PackageInfo.LocalPayloads.Num();
 	}
 
-	UE_LOG(LogVirtualization, Display, TEXT("Found %" INT64_FMT " payloads that need to be pushed to persistent virtualized storage"), TotalPayloadsToVirtualize);
+	UE_LOG(LogVirtualization, Display, TEXT("Found %" INT64_FMT " payload(s) that need to be pushed to persistent virtualized storage"), TotalPayloadsToVirtualize);
 
 	// TODO Optimization: In theory we could have many packages sharing the same payload and we only need to push once
+	
+	TArray<Virtualization::FPushRequest> PayloadsToSubmit;
+	PayloadsToSubmit.Reserve(TotalPayloadsToVirtualize);
 
 	// Push any remaining local payload to the persistent backends
 	for (FPackageInfo& PackageInfo : Packages)
@@ -261,6 +271,8 @@ void OnPrePackageSubmission(const TArray<FString>& FilesToSubmit, TArray<FText>&
 			return;
 		}
 
+		PackageInfo.PayloadIndex = PayloadsToSubmit.Num();
+
 		for (const FPayloadId& PayloadId : PackageInfo.LocalPayloads)
 		{
 			checkf(PayloadId.IsValid(), TEXT("PackageTrailer for package '%s' should not contain invalid FPayloadIds"), *PackageInfo.Path.GetDebugName());
@@ -277,37 +289,49 @@ void OnPrePackageSubmission(const TArray<FString>& FilesToSubmit, TArray<FText>&
 				return;
 			}
 
-			if (Payload)
-			{
-				if (!System.PushData(PayloadId, Payload, EStorageType::Persistent, PackageInfo.Path))
-				{
-					FText Message = FText::Format(	LOCTEXT("Virtualization_PushFailure", "Failed to push payload '{0}' from the package '{1}'"),
-													FText::FromString(PayloadId.ToString()),
-													FText::FromString(PackageInfo.Path.GetDebugName()));
-					Errors.Add(Message);
-					return;
-				}
-			}
-			else
+			if (!Payload)
 			{
 				FText Message = FText::Format(	LOCTEXT("Virtualization_MissingPayload", "Unable to find the payload '{0}' in the local storage of package '{1}'"),
 												FText::FromString(PayloadId.ToString()),
 												FText::FromString(PackageInfo.Path.GetDebugName()));
 				Errors.Add(Message);
 				return;
-			}
 
-			if (PackageInfo.Trailer.UpdatePayloadAsVirtualized(PayloadId))
-			{
-				PackageInfo.bWasTrailerUpdated = true;
 			}
-			else
+
+			PayloadsToSubmit.Emplace(PayloadId, MoveTemp(Payload), PackageInfo.Path.GetDebugName());
+		}
+	}
+
+	if (!System.PushData(PayloadsToSubmit, EStorageType::Persistent))
+	{
+		FText Message = LOCTEXT("Virtualization_PushFailure", "Failed to push payloads");
+		Errors.Add(Message);
+		return;
+	}
+
+	// Update the package info for the submitted payloads
+	for (FPackageInfo& PackageInfo : Packages)
+	{
+		for (int32 Index = 0; Index < PackageInfo.LocalPayloads.Num(); ++Index)
+		{
+			const Virtualization::FPushRequest& Request = PayloadsToSubmit[PackageInfo.PayloadIndex + Index];
+			check(Request.Identifier == PackageInfo.LocalPayloads[Index]);
+
+			if (Request.Status == Virtualization::FPushRequest::EStatus::Success)
 			{
-				FText Message = FText::Format(	LOCTEXT("Virtualization_UpdateStatusFailed", "Unable to update the status for the payload '{0}' in the package '{1}'"),
-												FText::FromString(PayloadId.ToString()),
-												FText::FromString(PackageInfo.Path.GetDebugName()));
-				Errors.Add(Message);
-				return;
+				if (PackageInfo.Trailer.UpdatePayloadAsVirtualized(Request.Identifier))
+				{
+					PackageInfo.bWasTrailerUpdated = true;
+				}
+				else
+					{
+					FText Message = FText::Format(	LOCTEXT("Virtualization_UpdateStatusFailed", "Unable to update the status for the payload '{0}' in the package '{1}'"),
+													FText::FromString(Request.Identifier.ToString()),
+													FText::FromString(PackageInfo.Path.GetDebugName()));
+					Errors.Add(Message);
+					return;
+				}
 			}
 		}
 	}
@@ -376,7 +400,7 @@ void OnPrePackageSubmission(const TArray<FString>& FilesToSubmit, TArray<FText>&
 		PackagesToReplace.Emplace(PackagePath, TempFilePath);
 	}
 
-	UE_LOG(LogVirtualization, Display, TEXT("%d packages had their trailer container modified and need to be updated"), PackagesToReplace.Num());
+	UE_LOG(LogVirtualization, Display, TEXT("%d package(s) had their trailer container modified and need to be updated"), PackagesToReplace.Num());
 
 	if (NumErrors == Errors.Num())
 	{
@@ -418,7 +442,7 @@ void OnPrePackageSubmission(const TArray<FString>& FilesToSubmit, TArray<FText>&
 		DescriptionTags.Add(Tag);
 	}
 
-	const double TimeInSeconds = FPlatformTime::ToSeconds64(FPlatformTime::Seconds() - StartTime);
+	const double TimeInSeconds = FPlatformTime::Seconds() - StartTime;
 	UE_LOG(LogVirtualization, Verbose, TEXT("Virtualization pre submit check took %.3f(s)"), TimeInSeconds);
 }
 
diff --git a/Engine/Source/Developer/Virtualization/Private/VirtualizationDDCBackend.cpp b/Engine/Source/Developer/Virtualization/Private/VirtualizationDDCBackend.cpp
--- a/Engine/Source/Developer/Virtualization/Private/VirtualizationDDCBackend.cpp
+++ b/Engine/Source/Developer/Virtualization/Private/VirtualizationDDCBackend.cpp
@@ -73,7 +73,7 @@ bool FDDCBackend::Initialize(const FString& ConfigEntry)
 	return true;	
 }
 
-EPushResult FDDCBackend::PushData(const FPayloadId& Id, const FCompressedBuffer& Payload, const FPackagePath& PackageContext)
+EPushResult FDDCBackend::PushData(const FPayloadId& Id, const FCompressedBuffer& Payload, const FString& PackageContext)
 {
 	TRACE_CPUPROFILER_EVENT_SCOPE(FDDCBackend::PushData);
 
diff --git a/Engine/Source/Developer/Virtualization/Private/VirtualizationDDCBackend.h b/Engine/Source/Developer/Virtualization/Private/VirtualizationDDCBackend.h
--- a/Engine/Source/Developer/Virtualization/Private/VirtualizationDDCBackend.h
+++ b/Engine/Source/Developer/Virtualization/Private/VirtualizationDDCBackend.h
@@ -36,7 +36,7 @@ private:
 
 	virtual bool Initialize(const FString& ConfigEntry) override;
 
-	virtual EPushResult PushData(const FPayloadId& Id, const FCompressedBuffer& Payload, const FPackagePath& PackageContext) override;
+	virtual EPushResult PushData(const FPayloadId& Id, const FCompressedBuffer& Payload, const FString& PackageContext) override;
 
 	virtual FCompressedBuffer PullData(const FPayloadId& Id) override;
 	
diff --git a/Engine/Source/Developer/Virtualization/Private/VirtualizationFileBackend.cpp b/Engine/Source/Developer/Virtualization/Private/VirtualizationFileBackend.cpp
--- a/Engine/Source/Developer/Virtualization/Private/VirtualizationFileBackend.cpp
+++ b/Engine/Source/Developer/Virtualization/Private/VirtualizationFileBackend.cpp
@@ -54,7 +54,7 @@ bool FFileSystemBackend::Initialize(const FString& ConfigEntry)
 	return true;
 }
 
-EPushResult FFileSystemBackend::PushData(const FPayloadId& Id, const FCompressedBuffer& Payload, const FPackagePath& PackageContext)
+EPushResult FFileSystemBackend::PushData(const FPayloadId& Id, const FCompressedBuffer& Payload, const FString& PackageContext)
 {
 	TRACE_CPUPROFILER_EVENT_SCOPE(FFileSystemBackend::PushData);
 
diff --git a/Engine/Source/Developer/Virtualization/Private/VirtualizationFileBackend.h b/Engine/Source/Developer/Virtualization/Private/VirtualizationFileBackend.h
--- a/Engine/Source/Developer/Virtualization/Private/VirtualizationFileBackend.h
+++ b/Engine/Source/Developer/Virtualization/Private/VirtualizationFileBackend.h
@@ -39,7 +39,7 @@ private:
 
 	virtual bool Initialize(const FString& ConfigEntry) override;
 
-	virtual EPushResult PushData(const FPayloadId& Id, const FCompressedBuffer& Payload, const FPackagePath& PackageContext) override;
+	virtual EPushResult PushData(const FPayloadId& Id, const FCompressedBuffer& Payload, const FString& PackageContext) override;
 
 	virtual FCompressedBuffer PullData(const FPayloadId& Id) override;
 
diff --git a/Engine/Source/Developer/Virtualization/Private/VirtualizationJupiterBackend.cpp b/Engine/Source/Developer/Virtualization/Private/VirtualizationJupiterBackend.cpp
--- a/Engine/Source/Developer/Virtualization/Private/VirtualizationJupiterBackend.cpp
+++ b/Engine/Source/Developer/Virtualization/Private/VirtualizationJupiterBackend.cpp
@@ -1374,7 +1374,7 @@ bool FJupiterBackend::Initialize(const FString& ConfigEntry)
 	return true;
 }
 
-EPushResult FJupiterBackend::PushData(const FPayloadId& Id, const FCompressedBuffer& CompressedPayload, const FPackagePath& PackageContext)
+EPushResult FJupiterBackend::PushData(const FPayloadId& Id, const FCompressedBuffer& CompressedPayload, const FString& PackageContext)
 {
 	using namespace Utility;
 
diff --git a/Engine/Source/Developer/Virtualization/Private/VirtualizationJupiterBackend.h b/Engine/Source/Developer/Virtualization/Private/VirtualizationJupiterBackend.h
--- a/Engine/Source/Developer/Virtualization/Private/VirtualizationJupiterBackend.h
+++ b/Engine/Source/Developer/Virtualization/Private/VirtualizationJupiterBackend.h
@@ -36,7 +36,7 @@ private:
 
 	virtual bool Initialize(const FString& ConfigEntry) override;
 
-	virtual EPushResult PushData(const FPayloadId& Id, const FCompressedBuffer& CompressedPayload, const FPackagePath& PackageContext) override;
+	virtual EPushResult PushData(const FPayloadId& Id, const FCompressedBuffer& CompressedPayload, const FString& PackageContext) override;
 
 	virtual FCompressedBuffer PullData(const FPayloadId& Id) override;
 
diff --git a/Engine/Source/Developer/Virtualization/Private/VirtualizationManager.cpp b/Engine/Source/Developer/Virtualization/Private/VirtualizationManager.cpp
--- a/Engine/Source/Developer/Virtualization/Private/VirtualizationManager.cpp
+++ b/Engine/Source/Developer/Virtualization/Private/VirtualizationManager.cpp
@@ -288,48 +288,64 @@ bool FVirtualizationManager::IsEnabled() const
 	return !AllBackends.IsEmpty();
 }
 
-bool FVirtualizationManager::PushData(const FPayloadId& Id, const FCompressedBuffer& Payload, EStorageType StorageType, const FPackagePath& PackageContext)
+bool FVirtualizationManager::PushData(const FPayloadId& Id, const FCompressedBuffer& Payload, EStorageType StorageType, const FString& Context)
+{
+	FPushRequest Request(Id, Payload, Context);
+	return FVirtualizationManager::PushData(MakeArrayView(&Request, 1), StorageType);
+}
+
+bool FVirtualizationManager::PushData(TArrayView<FPushRequest> Requests, EStorageType StorageType)
 {
 	TRACE_CPUPROFILER_EVENT_SCOPE(FVirtualizationManager::PushData);
 
-	if (!Id.IsValid())
-	{
-		// TODO: Should an invalid FPayloadId be an expected input, if so demote this from Warning->Verbose
-		UE_LOG(LogVirtualization, Warning, TEXT("Attempting to push a virtualized payload with an invalid FPayloadId"));
-		return false;
-	}
+	TArray<FPushRequest> ValidatedRequests;
+	ValidatedRequests.Reserve(Requests.Num());
 
-	FConditionalScopeLock _(&ForceSingleThreadedCS, bForceSingleThreaded);
+	TArray<int32> OriginalToValidatedRequest; // Builds a mapping between Requests and ValidatedRequests for later
+	OriginalToValidatedRequest.SetNum(Requests.Num());
 
-	// Early out if there are no backends or if the pushing of payloads has been disabled
-	if (!IsEnabled() || bEnablePayloadPushing == false)
+	// Create a new list of FPushRequest, excluding the requests that should not be processed for what ever reason.
+	for (int32 Index = 0; Index < Requests.Num(); ++Index)
 	{
-		return false;
-	}
+		OriginalToValidatedRequest[Index] = INDEX_NONE;
 
-	// Early out if we have no payload
-	if (Payload.GetCompressedSize() == 0)
-	{
-		// TODO: Should an invalid payload be an expected input, if so demote this from Warning->Verbose
-		UE_LOG(LogVirtualization, Warning, TEXT("Attempting to push an invalid virtualized payload (id: %s)"), *Id.ToString());
-		return false;
-	}
+		FPushRequest& Request = Requests[Index];
+		if (!Request.Identifier.IsValid() || Request.Payload.GetCompressedSize() == 0)
+		{
+			Request.Status = FPushRequest::EStatus::Invalid;
+			continue;
+		}
 
-	// Early out if the payload length is below our minimum required length
-	if ((int64)Payload.GetCompressedSize() < MinPayloadLength)
-	{
-		UE_LOG(	LogVirtualization, 
-				Verbose, 
-				TEXT("Attempting to push a virtualized payload (id: %s) that is smaller (%" UINT64_FMT ") than the MinPayloadLength (%" INT64_FMT ")"), 
-				*Id.ToString(), Payload.GetCompressedSize(), MinPayloadLength);
+		if ((int64)Request.Payload.GetCompressedSize() < MinPayloadLength)
+		{
+			UE_LOG(	LogVirtualization, Verbose, TEXT("Attempting to push a virtualized payload (id: %s) that is smaller (%" UINT64_FMT ") than the MinPayloadLength (%" INT64_FMT ")"),
+					*Request.Identifier.ToString(), 
+					Request.Payload.GetCompressedSize(),
+					MinPayloadLength);
+
+			Request.Status = FPushRequest::EStatus::BelowMinSize;
+			continue;
+		}
 
+		OriginalToValidatedRequest[Index] = ValidatedRequests.Num();
+		ValidatedRequests.Add(Request);
+	}
+	
+	// Filtering by package name is currently disabled
+#if 0
+	if (!ShouldVirtualizePackage(PackageContext))
+	{
+		UE_LOG(LogVirtualization, Verbose, TEXT("Payload '%s' for package '%s' will not be virtualized due to filtering"),
+			*Id.ToString(), *PackageContext.GetDebugName());
 		return false;
 	}
+#endif
 
-	if (!ShouldVirtualizePackage(PackageContext))
+	FConditionalScopeLock _(&ForceSingleThreadedCS, bForceSingleThreaded);
+
+	// Early out if there are no backends or if the pushing of payloads has been disabled
+	if (!IsEnabled() || bEnablePayloadPushing == false)
 	{
-		UE_LOG(	LogVirtualization, Verbose, TEXT("Payload '%s' for package '%s' will not be virtualized due to filtering"), 
-				*Id.ToString(), *PackageContext.GetDebugName());
 		return false;
 	}
 
@@ -343,29 +359,42 @@ bool FVirtualizationManager::PushData(const FPayloadId& Id, const FCompressedBuf
 
 	for (IVirtualizationBackend* Backend : Backends)
 	{
-		const EPushResult Result = TryPushDataToBackend(*Backend, Id, Payload, PackageContext) ? EPushResult::Success : EPushResult::Failed;
+		const bool bResult = TryPushDataToBackend(*Backend, ValidatedRequests);
 
-		UE_CLOG(Result != EPushResult::Failed, LogVirtualization, Verbose, TEXT("[%s] Pushed the payload '%s'"), *Backend->GetDebugName(), *Id.ToString());
-		UE_CLOG(Result == EPushResult::Failed, LogVirtualization, Error, TEXT("[%s] Failed to push the payload '%s'"), *Backend->GetDebugName(), *Id.ToString());
+		UE_CLOG(bResult == true, LogVirtualization, Verbose, TEXT("[%s] Pushed '%d' payload(s)"), *Backend->GetDebugName(), ValidatedRequests.Num());
+		UE_CLOG(bResult == false, LogVirtualization, Error, TEXT("[%s] Failed to push '%d' payload(s)"), *Backend->GetDebugName(), ValidatedRequests.Num());
 
-		if (Result == EPushResult::Failed)
+		if (!bResult)
 		{
 			ErrorCount++;
 		}
 		
 		// Debugging operation where we immediately try to pull the payload after each push (when possible) and assert 
 		// that the pulled payload is the same as the original
-		if (bValidateAfterPushOperation && Result != EPushResult::Failed && Backend->SupportsPullOperations())
+		if (bValidateAfterPushOperation && bResult == true && Backend->SupportsPullOperations())
 		{
-			FCompressedBuffer PulledPayload = PullDataFromBackend(*Backend, Id);
-			checkf(	Payload.GetRawHash() == PulledPayload.GetRawHash(), 
-					TEXT("[%s] Failed to pull payload '%s' after it was pushed to backend"), 
-					*Backend->GetDebugName(), 
-					*Id.ToString());
+			for (FPushRequest& Request : ValidatedRequests)
+			{
+				FCompressedBuffer ValidationPayload = PullDataFromBackend(*Backend, Request.Identifier);
+				checkf(	Request.Payload.GetRawHash() == ValidationPayload.GetRawHash(),
+						TEXT("[%s] Failed to pull payload '%s' after it was pushed to backend"),
+						*Backend->GetDebugName(),
+						*Request.Identifier.ToString());
+			}
 		}
 	}
 
-	UE_CLOG(ErrorCount == Backends.Num(), LogVirtualization, Error, TEXT("Payload '%s' failed to be pushed to any backend'"), *Id.ToString());
+	UE_CLOG(ErrorCount == Backends.Num(), LogVirtualization, Error, TEXT("Failed to push '%d' payload(s) to any backend'"), ValidatedRequests.Num());
+
+	// Now we need to update the statuses of the original list of requests with those from our validated list
+	for (int32 Index = 0; Index < Requests.Num(); ++Index)
+	{
+		const int32 MappingIndex = OriginalToValidatedRequest[Index];
+		if (MappingIndex != INDEX_NONE)
+		{
+			Requests[Index].Status = ValidatedRequests[MappingIndex].Status;
+		}
+	}
 
 	// For local storage we consider the push to have failed only if ALL backends gave an error, if at least one backend succeeded then the operation succeeded.
 	// For persistent storage we require that all backends succeeded, so any errors will fail the push operation.
@@ -848,7 +877,7 @@ void FVirtualizationManager::CachePayload(const FPayloadId& Id, const FCompresse
 bool FVirtualizationManager::TryCacheDataToBackend(IVirtualizationBackend& Backend, const FPayloadId& Id, const FCompressedBuffer& Payload)
 {
 	COOK_STAT(FCookStats::FScopedStatsCounter Timer(Profiling::GetCacheStats(Backend)));
-	const EPushResult Result = Backend.PushData(Id, Payload, FPackagePath());
+	const EPushResult Result = Backend.PushData(Id, Payload, FString());
 
 	if (Result == EPushResult::Success)
 	{
@@ -858,17 +887,24 @@ bool FVirtualizationManager::TryCacheDataToBackend(IVirtualizationBackend& Backe
 	return Result != EPushResult::Failed;
 }
 
-bool FVirtualizationManager::TryPushDataToBackend(IVirtualizationBackend& Backend, const FPayloadId& Id, const FCompressedBuffer& Payload, const FPackagePath& PackageContext)
+bool FVirtualizationManager::TryPushDataToBackend(IVirtualizationBackend& Backend, TArrayView<FPushRequest> Requests)
 {
 	COOK_STAT(FCookStats::FScopedStatsCounter Timer(Profiling::GetPushStats(Backend)));
-	const EPushResult Result = Backend.PushData(Id, Payload, PackageContext);
+	const bool bResult = Backend.PushData(Requests);
 
-	if (Result == EPushResult::Success)
+	if (bResult)
 	{
-		COOK_STAT(Timer.AddHit(Payload.GetCompressedSize()));
+		for (FPushRequest& Request : Requests)
+		{
+			// TODO: Don't add a hit if the payload was already uploaded
+			if (Request.Status == FPushRequest::EStatus::Success)
+			{
+				COOK_STAT(Timer.AddHit(Request.Payload.GetCompressedSize()));
+			}
+		}
 	}
-
-	return Result != EPushResult::Failed;
+	
+	return bResult;
 }
 
 FCompressedBuffer FVirtualizationManager::PullDataFromBackend(IVirtualizationBackend& Backend, const FPayloadId& Id)
diff --git a/Engine/Source/Developer/Virtualization/Private/VirtualizationManager.h b/Engine/Source/Developer/Virtualization/Private/VirtualizationManager.h
--- a/Engine/Source/Developer/Virtualization/Private/VirtualizationManager.h
+++ b/Engine/Source/Developer/Virtualization/Private/VirtualizationManager.h
@@ -93,7 +93,8 @@ private:
 
 	virtual bool IsEnabled() const override;
 	
-	virtual bool PushData(const FPayloadId& Id, const FCompressedBuffer& Payload, EStorageType StorageType, const FPackagePath& PackageContext) override;
+	virtual bool PushData(const FPayloadId& Id, const FCompressedBuffer& Payload, EStorageType StorageType, const FString& Context) override;
+	virtual bool PushData(TArrayView<FPushRequest> Requests, EStorageType StorageType) override;
 
 	virtual FCompressedBuffer PullData(const FPayloadId& Id) override;
 
@@ -125,7 +126,8 @@ private:
 	void CachePayload(const FPayloadId& Id, const FCompressedBuffer& Payload, const IVirtualizationBackend* BackendSource);
 
 	bool TryCacheDataToBackend(IVirtualizationBackend& Backend, const FPayloadId& Id, const FCompressedBuffer& Payload);
-	bool TryPushDataToBackend(IVirtualizationBackend& Backend, const FPayloadId& Id, const FCompressedBuffer& Payload, const FPackagePath& PackageContext);
+	bool TryPushDataToBackend(IVirtualizationBackend& Backend, TArrayView<FPushRequest> Requests);
+
 	FCompressedBuffer PullDataFromBackend(IVirtualizationBackend& Backend, const FPayloadId& Id);
 
 	/** 
diff --git a/Engine/Source/Developer/Virtualization/Private/VirtualizationSourceControlBackend.cpp b/Engine/Source/Developer/Virtualization/Private/VirtualizationSourceControlBackend.cpp
--- a/Engine/Source/Developer/Virtualization/Private/VirtualizationSourceControlBackend.cpp
+++ b/Engine/Source/Developer/Virtualization/Private/VirtualizationSourceControlBackend.cpp
@@ -22,16 +22,26 @@ namespace UE::Virtualization
 {
 
 /** Builds a changelist description to be used when submitting a payload to source control */
-void CreateDescription(const FPackagePath& PackageContext, TStringBuilder<512>& OutDescription)
+void CreateDescription(const TArray<const FPushRequest*>& FileRequests, TStringBuilder<512>& OutDescription)
 {
 	// TODO: Maybe make writing out the project name an option or allow for a codename to be set via ini file?
 	OutDescription << TEXT("Submitted for project: ");
 	OutDescription << FApp::GetProjectName();
 
-	if (PackageContext.IsMountedPath())
+	bool bInitialNewline = false;
+
+	for (const FPushRequest* Request : FileRequests)
 	{
-		OutDescription << TEXT("\nPackage: ");
-		PackageContext.AppendPackageName(OutDescription);
+		if (!Request->Context.IsEmpty())
+		{
+			if (!bInitialNewline)
+			{
+				OutDescription << TEXT("\n");
+				bInitialNewline = true;
+			}
+
+			OutDescription << TEXT("\n") << Request->Identifier << "\t: " << Request->Context;
+		}
 	}
 }
 
@@ -113,7 +123,57 @@ bool FSourceControlBackend::Initialize(const FString& ConfigEntry)
 	return true;
 }
 
-EPushResult FSourceControlBackend::PushData(const FPayloadId& Id, const FCompressedBuffer& Payload, const FPackagePath& PackageContext)
+FCompressedBuffer FSourceControlBackend::PullData(const FPayloadId& Id)
+{
+	TRACE_CPUPROFILER_EVENT_SCOPE(FSourceControlBackend::PullData);
+
+	TStringBuilder<512> DepotPath;
+	CreateDepotPath(Id, DepotPath);
+
+	ISourceControlProvider& SCCProvider = ISourceControlModule::Get().GetProvider();
+
+#if IS_SOURCE_CONTROL_THREAD_SAFE
+	TSharedRef<FDownloadFile, ESPMode::ThreadSafe> DownloadCommand = ISourceControlOperation::Create<FDownloadFile>(FDownloadFile::EVerbosity::None);
+	if (SCCProvider.Execute(DownloadCommand, DepotPath.ToString(), EConcurrency::Synchronous) != ECommandResult::Succeeded)
+	{
+		return FCompressedBuffer();
+	}
+#else
+	TSharedRef<FDownloadFile> DownloadCommand = ISourceControlOperation::Create<FDownloadFile>(FDownloadFile::EVerbosity::None);
+	if (!SCCProvider.TryToDownloadFileFromBackgroundThread(DownloadCommand, DepotPath.ToString()))
+	{
+		return FCompressedBuffer();
+	}
+#endif
+
+	// The payload was created by FCompressedBuffer::Compress so we can return it 
+	// as a FCompressedBuffer.
+	FSharedBuffer Buffer = DownloadCommand->GetFileData(DepotPath);
+	return FCompressedBuffer::FromCompressed(Buffer);
+}
+
+bool FSourceControlBackend::DoesPayloadExist(const FPayloadId& Id)
+{
+	TArray<bool> Result;
+
+	if (FSourceControlBackend::DoPayloadsExist(MakeArrayView<const FPayloadId>(&Id, 1), Result))
+	{
+		check(Result.Num() == 1);
+		return Result[0];
+	}
+	else
+	{
+		return false;
+	}
+}
+
+EPushResult FSourceControlBackend::PushData(const FPayloadId& Id, const FCompressedBuffer& Payload, const FString& Context)
+{
+	FPushRequest Request(Id, Payload, Context);
+	return FSourceControlBackend::PushData(MakeArrayView(&Request, 1)) ? EPushResult::Success : EPushResult::Failed;
+}
+
+bool FSourceControlBackend::PushData(TArrayView<FPushRequest> Requests)
 {
 	TRACE_CPUPROFILER_EVENT_SCOPE(FSourceControlBackend::PushData);
 
@@ -131,56 +191,65 @@ EPushResult FSourceControlBackend::PushData(const FPayloadId& Id, const FCompres
 	TStringBuilder<260> RootDirectory;
 	RootDirectory << FPlatformProcess::UserTempDir() << TEXT("UnrealEngine/VirtualizedPayloads/") << SessionGuid << TEXT("/");
 
-	// First we need to save the payload to a file in the workspace client mapping so that it can be submitted
-	TStringBuilder<52> LocalPayloadPath;
-	Utils::PayloadIdToPath(Id, LocalPayloadPath);
-
-	FString PayloadFilePath = *WriteToString<512>(RootDirectory, LocalPayloadPath);
-
 	ON_SCOPE_EXIT
 	{
 		// Clean up the payload file from disk and the temp directories, but we do not need to give errors if any of these operations fail.
-		IFileManager::Get().Delete(*PayloadFilePath, false, false, true);
 		IFileManager::Get().DeleteDirectory(RootDirectory.ToString(), false, true);
 	};
 
-	// Write the payload to Disk
+	TArray<FString> FilesToSubmit;
+	FilesToSubmit.Reserve(Requests.Num());
+
+	// Write the payloads to disk so that they can be submitted
 	{
-		UE_LOG(LogVirtualization, Verbose, TEXT("[%s] Writing payload to '%s' for submission"), *GetDebugName(), *PayloadFilePath);
+		TRACE_CPUPROFILER_EVENT_SCOPE(FSourceControlBackend::PushData::CreateFiles);
 
-		TUniquePtr<FArchive> FileAr(IFileManager::Get().CreateFileWriter(*PayloadFilePath));
-		if (!FileAr)
+		for (const FPushRequest& Request : Requests)
 		{
-			TStringBuilder<MAX_SPRINTF> SystemErrorMsg;
-			Utils::GetFormattedSystemError(SystemErrorMsg);
+			TStringBuilder<52> LocalPayloadPath;
+			Utils::PayloadIdToPath(Request.Identifier, LocalPayloadPath);
+			FString PayloadFilePath = *WriteToString<512>(RootDirectory, LocalPayloadPath);
 
-			UE_LOG(LogVirtualization, Error, TEXT("[%s] Failed to write payload '%s' contents to '%s' due to system error: %s"),
-				*GetDebugName(),
-				*Id.ToString(),
-				*PayloadFilePath,
-				SystemErrorMsg.ToString());
+			UE_LOG(LogVirtualization, Verbose, TEXT("[%s] Writing payload to '%s' for submission"), *GetDebugName(), *PayloadFilePath);
 
-			return EPushResult::Failed;
-		}
+			TUniquePtr<FArchive> FileAr(IFileManager::Get().CreateFileWriter(*PayloadFilePath));
+			if (!FileAr)
+			{
+				TStringBuilder<MAX_SPRINTF> SystemErrorMsg;
+				Utils::GetFormattedSystemError(SystemErrorMsg);
+
+				UE_LOG(LogVirtualization, Error, TEXT("[%s] Failed to write payload '%s' contents to '%s' due to system error: %s"),
+					*GetDebugName(),
+					*Request.Identifier.ToString(),
+					*PayloadFilePath,
+					SystemErrorMsg.ToString());
 
-		Payload.Save(*FileAr);
+				return false;
+			}
 
-		if (!FileAr->Close())
-		{
-			TStringBuilder<MAX_SPRINTF> SystemErrorMsg;
-			Utils::GetFormattedSystemError(SystemErrorMsg);
+			Request.Payload.Save(*FileAr);
 
-			UE_LOG(LogVirtualization, Error, TEXT("[%s] Failed to write payload '%s' contents to '%s' due to system error: %s"),
-				*GetDebugName(),
-				*Id.ToString(),
-				*PayloadFilePath,
-				SystemErrorMsg.
-				ToString());
+			if (!FileAr->Close())
+			{
+				TStringBuilder<MAX_SPRINTF> SystemErrorMsg;
+				Utils::GetFormattedSystemError(SystemErrorMsg);
+
+				UE_LOG(LogVirtualization, Error, TEXT("[%s] Failed to write payload '%s' contents to '%s' due to system error: %s"),
+					*GetDebugName(),
+					*Request.Identifier.ToString(),
+					*PayloadFilePath,
+					SystemErrorMsg.
+					ToString());
+
+				return false;
+			}
 
-			return EPushResult::Failed;
+			FilesToSubmit.Emplace(MoveTemp(PayloadFilePath));
 		}
 	}
 
+	check(Requests.Num() == FilesToSubmit.Num());
+
 	TStringBuilder<64> WorkspaceName;
 	WorkspaceName << TEXT("MirageSubmission-") << SessionGuid;
 
@@ -188,6 +257,7 @@ EPushResult FSourceControlBackend::PushData(const FPayloadId& Id, const FCompres
 
 	// Create a temp workspace so that we can submit the payload from
 	{
+		TRACE_CPUPROFILER_EVENT_SCOPE(FSourceControlBackend::PushData::CreateWorkspace);
 		TSharedRef<FCreateWorkspace> CreateWorkspaceCommand = ISourceControlOperation::Create<FCreateWorkspace>(WorkspaceName, RootDirectory);
 
 		TStringBuilder<512> DepotMapping;
@@ -205,12 +275,11 @@ EPushResult FSourceControlBackend::PushData(const FPayloadId& Id, const FCompres
 
 		if (SCCProvider.Execute(CreateWorkspaceCommand) != ECommandResult::Succeeded)
 		{
-			UE_LOG(LogVirtualization, Error, TEXT("[%s] Failed to create temp workspace '%s' to submit payload '%s' from"),
+			UE_LOG(LogVirtualization, Error, TEXT("[%s] Failed to create temp workspace '%s' to submit payloads from"),
 				*GetDebugName(),
-				WorkspaceName.ToString(),
-				*Id.ToString());
+				WorkspaceName.ToString());
 
-			return EPushResult::Failed;
+			return false;
 		}
 	}
 
@@ -223,16 +292,19 @@ EPushResult FSourceControlBackend::PushData(const FPayloadId& Id, const FCompres
 		}
 	};
 
-	FSourceControlResultInfo SwitchToNewWorkspaceInfo;
 	FString OriginalWorkspace;
-	if (SCCProvider.SwitchWorkspace(WorkspaceName, SwitchToNewWorkspaceInfo, &OriginalWorkspace) != ECommandResult::Succeeded)
 	{
-		UE_LOG(LogVirtualization, Error, TEXT("[%s] Failed to switch to temp workspace '%s' when trying to submit payload '%s'"),
-			*GetDebugName(),
-			WorkspaceName.ToString(),
-			*Id.ToString());
+		TRACE_CPUPROFILER_EVENT_SCOPE(FSourceControlBackend::PushData::SwitchWorkspace);
 
-		return EPushResult::Failed;
+		FSourceControlResultInfo SwitchToNewWorkspaceInfo;
+		if (SCCProvider.SwitchWorkspace(WorkspaceName, SwitchToNewWorkspaceInfo, &OriginalWorkspace) != ECommandResult::Succeeded)
+		{
+			UE_LOG(LogVirtualization, Error, TEXT("[%s] Failed to switch to temp workspace '%s' when trying to submit payloads"),
+				*GetDebugName(),
+				WorkspaceName.ToString());
+
+			return false;
+		}
 	}
 
 	ON_SCOPE_EXIT
@@ -249,94 +321,85 @@ EPushResult FSourceControlBackend::PushData(const FPayloadId& Id, const FCompres
 		}
 	};
 
-	FSourceControlStatePtr FileState = SCCProvider.GetState(PayloadFilePath, EStateCacheUsage::ForceUpdate);
-	if (!FileState.IsValid())
+	TArray<FSourceControlStateRef> FileStates;
 	{
-		UE_LOG(LogVirtualization, Error, TEXT("[%s] Failed to find the current file state for '%s'"), *GetDebugName(), *PayloadFilePath);
-		return EPushResult::Failed;
+		TRACE_CPUPROFILER_EVENT_SCOPE(FSourceControlBackend::PushData::GetFileStates);
+		if (SCCProvider.GetState(FilesToSubmit, FileStates, EStateCacheUsage::ForceUpdate) != ECommandResult::Succeeded)
+		{
+			UE_LOG(LogVirtualization, Error, TEXT("[%s] Failed to find the current file state for payloads"), *GetDebugName());
+			return false;
+		}
 	}
+	check(Requests.Num() == FileStates.Num());
 
-	if (FileState->IsSourceControlled())
-	{
-		// TODO: Maybe check if the data is the same (could be different if the compression algorithm has changed)
-		// TODO: Should we respect if the file is deleted as technically we can still get access to it?
-		return EPushResult::PayloadAlreadyExisted;
-	}
-	else if (FileState->CanAdd())
+	TArray<FString> FilesToAdd;
+	FilesToAdd.Reserve(FilesToSubmit.Num());
+
+	TArray<const FPushRequest*> FileRequests;
+	FileRequests.Reserve(FilesToSubmit.Num());
+
+	for (int32 Index = 0; Index < FilesToSubmit.Num(); ++Index)
 	{
-		if (SCCProvider.Execute(ISourceControlOperation::Create<FMarkForAdd>(), PayloadFilePath) != ECommandResult::Succeeded)
+		if (FileStates[Index]->IsSourceControlled())
 		{
-			UE_LOG(LogVirtualization, Error, TEXT("[%s] Failed to mark the payload file '%s' for Add in source control"), *GetDebugName(), *PayloadFilePath);
-			return EPushResult::Failed;
+			// TODO: Maybe check if the data is the same (could be different if the compression algorithm has changed)
+			// TODO: Should we respect if the file is deleted as technically we can still get access to it?
+			Requests[Index].Status = FPushRequest::EStatus::Success;
+		}
+		else if (FileStates[Index]->CanAdd())
+		{
+			FilesToAdd.Add(FilesToSubmit[Index]);
+			FileRequests.Add(&Requests[Index]);
+		}
+		else
+		{
+			UE_LOG(LogVirtualization, Error, TEXT("[%s] The the payload file '%s' is not in source control but also cannot be marked for Add"), *GetDebugName(), *FilesToSubmit[Index]);
+			return false;
 		}
 	}
-	else
+
+	check(FileRequests.Num() == FilesToAdd.Num());
+
+	if (FilesToAdd.IsEmpty())
 	{
-		UE_LOG(LogVirtualization, Error, TEXT("[%s] The the payload file '%s' is not in source control but also cannot be marked for Add"), *GetDebugName(), *PayloadFilePath);
-		return EPushResult::Failed;
+		return true;
+	}
+
+	{
+		TRACE_CPUPROFILER_EVENT_SCOPE(FSourceControlBackend::PushData::AddFiles);
+
+		if (SCCProvider.Execute(ISourceControlOperation::Create<FMarkForAdd>(), FilesToAdd) != ECommandResult::Succeeded)
+		{
+			UE_LOG(LogVirtualization, Error, TEXT("[%s] Failed to mark the payload file for Add in source control"), *GetDebugName());
+			return false;
+		}
 	}
 
 	// Now submit the payload
 	{
+		TRACE_CPUPROFILER_EVENT_SCOPE(FSourceControlBackend::PushData::SubmitFiles);
+
 		TSharedRef<FCheckIn, ESPMode::ThreadSafe> CheckInOperation = ISourceControlOperation::Create<FCheckIn>();
 
 		TStringBuilder<512> Description;
-		CreateDescription(PackageContext, Description);
+		CreateDescription(FileRequests, Description);
 
 		CheckInOperation->SetDescription(FText::FromString(Description.ToString()));
 
-		if (SCCProvider.Execute(CheckInOperation, PayloadFilePath) != ECommandResult::Succeeded)
+		if (SCCProvider.Execute(CheckInOperation, FilesToAdd) != ECommandResult::Succeeded)
 		{
-			UE_LOG(LogVirtualization, Error, TEXT("[%s] Failed to submit the payload file '%s' to source control"), *GetDebugName(), *PayloadFilePath);
-			return EPushResult::Failed;
+			UE_LOG(LogVirtualization, Error, TEXT("[%s] Failed to submit the payload file(s) to source control"), *GetDebugName());
+			return false;
 		}
 	}
 
-	return EPushResult::Success;
-}
-
-FCompressedBuffer FSourceControlBackend::PullData(const FPayloadId& Id)
-{
-	TRACE_CPUPROFILER_EVENT_SCOPE(FSourceControlBackend::PullData);
-
-	TStringBuilder<512> DepotPath;
-	CreateDepotPath(Id, DepotPath);
-
-	ISourceControlProvider& SCCProvider = ISourceControlModule::Get().GetProvider();
-
-#if IS_SOURCE_CONTROL_THREAD_SAFE
-	TSharedRef<FDownloadFile, ESPMode::ThreadSafe> DownloadCommand = ISourceControlOperation::Create<FDownloadFile>(FDownloadFile::EVerbosity::None);
-	if (SCCProvider.Execute(DownloadCommand, DepotPath.ToString(), EConcurrency::Synchronous) != ECommandResult::Succeeded)
-	{
-		return FCompressedBuffer();
-	}
-#else
-	TSharedRef<FDownloadFile> DownloadCommand = ISourceControlOperation::Create<FDownloadFile>(FDownloadFile::EVerbosity::None);
-	if (!SCCProvider.TryToDownloadFileFromBackgroundThread(DownloadCommand, DepotPath.ToString()))
+	// TODO: We really should be setting a more fine grain status for each request, or not bother with the status at all
+	for (FPushRequest& Request : Requests)
 	{
-		return FCompressedBuffer();
+		Request.Status = FPushRequest::EStatus::Success;
 	}
-#endif
-
-	// The payload was created by FCompressedBuffer::Compress so we can return it 
-	// as a FCompressedBuffer.
-	FSharedBuffer Buffer = DownloadCommand->GetFileData(DepotPath);
-	return FCompressedBuffer::FromCompressed(Buffer);
-}
-
-bool FSourceControlBackend::DoesPayloadExist(const FPayloadId& Id)
-{
-	TArray<bool> Result;
 
-	if (FSourceControlBackend::DoPayloadsExist(MakeArrayView<const FPayloadId>(&Id, 1), Result))
-	{
-		check(Result.Num() == 1);
-		return Result[0];
-	}
-	else
-	{
-		return false;
-	}
+	return true;
 }
 
 bool FSourceControlBackend::DoPayloadsExist(TArrayView<const FPayloadId> PayloadIds, TArray<bool>& OutResults)
diff --git a/Engine/Source/Developer/Virtualization/Private/VirtualizationSourceControlBackend.h b/Engine/Source/Developer/Virtualization/Private/VirtualizationSourceControlBackend.h
--- a/Engine/Source/Developer/Virtualization/Private/VirtualizationSourceControlBackend.h
+++ b/Engine/Source/Developer/Virtualization/Private/VirtualizationSourceControlBackend.h
@@ -38,8 +38,9 @@ private:
 
 	virtual bool Initialize(const FString& ConfigEntry) override;
 	
-	virtual EPushResult PushData(const FPayloadId& Id, const FCompressedBuffer& Payload, const FPackagePath& PackageContext) override;
-	
+	virtual EPushResult PushData(const FPayloadId& Id, const FCompressedBuffer& Payload, const FString& Context) override;
+	virtual bool PushData(TArrayView<FPushRequest> Requests) override;
+
 	virtual FCompressedBuffer PullData(const FPayloadId& Id) override;
 
 	virtual bool DoesPayloadExist(const FPayloadId& Id) override;
diff --git a/Engine/Source/Developer/Virtualization/Public/IVirtualizationBackend.h b/Engine/Source/Developer/Virtualization/Public/IVirtualizationBackend.h
--- a/Engine/Source/Developer/Virtualization/Public/IVirtualizationBackend.h
+++ b/Engine/Source/Developer/Virtualization/Public/IVirtualizationBackend.h
@@ -7,14 +7,12 @@
 #include "Features/IModularFeature.h"
 #include "Features/IModularFeatures.h"
 #include "Templates/UniquePtr.h"
-
-class FPackagePath;
+#include "Virtualization/PayloadId.h"
+#include "Virtualization/VirtualizationSystem.h"
 
 namespace UE::Virtualization
 {
 
-class FPayloadId;
-
 /** Describes the result of a IVirtualizationBackend::Push operation */
 enum class EPushResult
 {
@@ -86,7 +84,31 @@ public:
 	 * @param Payload	A potentially compressed buffer representing the payload
 	 * @return			The result of the push operation
 	 */
-	virtual EPushResult PushData(const FPayloadId& Id, const FCompressedBuffer& Payload, const FPackagePath& PackageContext) = 0;
+	virtual EPushResult PushData(const FPayloadId& Id, const FCompressedBuffer& Payload, const FString& PackageContext) = 0;
+
+	virtual bool PushData(TArrayView<FPushRequest> Requests)
+	{
+		// TODO: Sort return codes
+		for (FPushRequest& Request : Requests)
+		{
+			EPushResult Result = PushData(Request.Identifier, Request.Payload, Request.Context);
+			switch (Result)
+			{
+			case EPushResult::Failed:
+				return false;
+
+			case EPushResult::PayloadAlreadyExisted:
+			case EPushResult::Success:
+				break;
+
+			default:
+				checkNoEntry();
+				break;
+			}
+		}
+
+		return true;
+	}
 
 	/** 
 	 * The backend will attempt to retrieve the given payload by what ever method the backend uses.
diff --git a/Engine/Source/Runtime/Core/Private/Virtualization/VirtualizationSystem.cpp b/Engine/Source/Runtime/Core/Private/Virtualization/VirtualizationSystem.cpp
--- a/Engine/Source/Runtime/Core/Private/Virtualization/VirtualizationSystem.cpp
+++ b/Engine/Source/Runtime/Core/Private/Virtualization/VirtualizationSystem.cpp
@@ -24,7 +24,12 @@ public:
 		return false;
 	}
 
-	virtual bool PushData(const FPayloadId& Id, const FCompressedBuffer& Payload, EStorageType StorageType, const FPackagePath& Package) override
+	virtual bool PushData(const FPayloadId& Id, const FCompressedBuffer& Payload, EStorageType StorageType, const FString& Context) override
+	{
+		return false;
+	}
+
+	virtual bool PushData(TArrayView<FPushRequest> Requests, EStorageType StorageType) override
 	{
 		return false;
 	}
diff --git a/Engine/Source/Runtime/Core/Public/Virtualization/VirtualizationSystem.h b/Engine/Source/Runtime/Core/Public/Virtualization/VirtualizationSystem.h
--- a/Engine/Source/Runtime/Core/Public/Virtualization/VirtualizationSystem.h
+++ b/Engine/Source/Runtime/Core/Public/Virtualization/VirtualizationSystem.h
@@ -3,6 +3,7 @@
 #pragma once
 
 #include "Compression/CompressedBuffer.h"
+#include "Containers/StringFwd.h"
 #include "Features/IModularFeature.h"
 #include "Features/IModularFeatures.h"
 #include "Virtualization/PayloadId.h"
@@ -39,6 +40,7 @@ enum class EStorageType : int8
 	Persistent
 };
 
+/** Describes the status of a payload in regards to a backend storage system */
 enum class FPayloadStatus : int8
 {
 	/** The payload id was not value */
@@ -51,10 +53,52 @@ enum class FPayloadStatus : int8
 	FoundAll
 };
 
+/** Data structure representing a request to push a payload to a backend storage system */
+struct FPushRequest
+{
+	enum class EStatus
+	{
+		/** The request failed, or was not reached because of the failure of an earlier request */
+		Failed,
+		/** The payload does not have a valid identifier or is empty */
+		Invalid,
+		/** The payload is below the minimum length required for virtualization */
+		BelowMinSize,
+		/** The payload in the request is now present in all backends */
+		Success
+	};
+
+	FPushRequest(const FPayloadId& InIdentifier, const FCompressedBuffer& InPayload, const FString& InContext)
+		: Identifier(InIdentifier)
+		, Payload(InPayload)
+		, Context(InContext)
+	{
+
+	}
+
+	FPushRequest(const FPayloadId& InIdentifier, FCompressedBuffer&& InPayload, FString&& InContext)
+		: Identifier(InIdentifier)
+		, Payload(InPayload)
+		, Context(InContext)
+	{
+
+	}
+
+	/** The identifier of the payload */
+	FPayloadId Identifier;
+	/** The payload data */
+	FCompressedBuffer Payload;
+	/** A string containing context for the payload, typically a package name */
+	FString Context;
+
+	/** Once the request has been processed this value will contains the results */
+	EStatus Status = EStatus::Failed;
+};
+
 /**
  * Creates the global IVirtualizationSystem if it has not already been set up. This can be called explicitly
  * during process start up but it will also be called by IVirtualizationSystem::Get if it detects that the
- * IVirtualizationSystem has not yet bene set up.
+ * IVirtualizationSystem has not yet been set up.
  */
 CORE_API void Initialize();
 
@@ -91,11 +135,27 @@ public:
 	 * @param	Payload			The payload itself in FCompressedBuffer form, it is assumed that if the buffer is to 
 	 *							be compressed that it will have been done by the caller.
 	 * @param	StorageType		The type of storage to push the payload to, @See EStorageType for details.
-	 * @param	PackageContext	Name of the owning package, which can be used to provide context about the payload.
+	 * @param	PackageContext	Context for the payload being submitted, typically the name from the UPackage that owns it.
 	 * 
 	 * @return	True if at least one backend now contains the payload, otherwise false.
 	 */
-	virtual bool PushData(const FPayloadId& Id, const FCompressedBuffer& Payload, EStorageType StorageType, const FPackagePath& PackageContext) = 0;
+	virtual bool PushData(const FPayloadId& Id, const FCompressedBuffer& Payload, EStorageType StorageType, const FString& Context) = 0;
+	
+	/**
+	 * Push one or more payloads to a backend storage system. @See FPushRequest.
+	 * 
+	 * @param	Requests	A list of one or more payloads
+	 * @param	StorageType	The type of storage to push the payload to, @See EStorageType for details.
+	 * 
+	 * @return	When StorageType is Local, this method will return true assuming at least one backend
+	 *			managed to push all of the payloads. 
+	 *			When StorageType is Persistent, this method will only return true if ALL backends
+	 *			manage to push all of the payloads.
+	 * 			If this returns true then you can check the Status member of each request for more info 
+	 *			about each payloads push operation.
+	 *			If this returns false then you can assume that the payloads are not safely virtualized.
+	 */
+	virtual bool PushData(TArrayView<FPushRequest> Requests, EStorageType StorageType) = 0;
 
 	/**
 	 * Pull a payload from the virtualization backends.
diff --git a/Engine/Source/Runtime/CoreUObject/Private/Serialization/VirtualizedBulkData.cpp b/Engine/Source/Runtime/CoreUObject/Private/Serialization/VirtualizedBulkData.cpp
--- a/Engine/Source/Runtime/CoreUObject/Private/Serialization/VirtualizedBulkData.cpp
+++ b/Engine/Source/Runtime/CoreUObject/Private/Serialization/VirtualizedBulkData.cpp
@@ -1149,7 +1149,7 @@ void FVirtualizedUntypedBulkData::PushData(const FPackagePath& InPackagePath)
 		RecompressForSerialization(PayloadToPush, Flags);
 
 		// TODO: We could make this a config option?
-		if (VirtualizationSystem.PushData(PayloadContentId, PayloadToPush, EStorageType::Local, InPackagePath))
+		if (VirtualizationSystem.PushData(PayloadContentId, PayloadToPush, EStorageType::Local, InPackagePath.GetPackageName()))
 		{
 			EnumAddFlags(Flags, EFlags::IsVirtualized);
 			EnumRemoveFlags(Flags, EFlags::ReferencesLegacyFile | EFlags::LegacyFileIsCompressed);
