Rename FJupiterBackend to FHttpBackend and reference sto Jupiter have been changed to Horde Storage

#rnx

- This backend is still most likely going to be deleted before 5.0 ships but just in case it is not I am doing the renaming now to keep in step with the names we are using for Jupiter when communicating with externals.

diff --git a/Engine/Source/Developer/Virtualization/Private/VirtualizationHttpBackend.cpp b/Engine/Source/Developer/Virtualization/Private/VirtualizationHttpBackend.cpp
new file mode 100644
--- /dev/null
+++ b/Engine/Source/Developer/Virtualization/Private/VirtualizationHttpBackend.cpp
@@ -0,0 +1,2015 @@
+// Copyright Epic Games, Inc. All Rights Reserved.
+
+#include "VirtualizationHttpBackend.h"
+#include "Virtualization/PayloadId.h"
+
+// TODO: Our libcurl implementation does not currently support MacOS
+// (not registering this will cause a fatal log error if the backend is actually used)
+#if PLATFORM_WINDOWS
+	#define UE_USE_HORDESTORAGE_BACKEND 1
+#else
+	#define UE_USE_HORDESTORAGE_BACKEND 0
+#endif //PLATFORM_WINDOWS
+
+#if UE_USE_HORDESTORAGE_BACKEND
+
+#include "Async/TaskGraphInterfaces.h"
+#include "Compression/CompressedBuffer.h"
+#include "Containers/StringConv.h"
+#include "Containers/Ticker.h"
+#include "Logging/LogMacros.h"
+#include "Misc/FileHelper.h"
+#include "Misc/Parse.h"
+#include "Misc/ScopeLock.h"
+#include "Misc/StringBuilder.h"
+#include "Serialization/JsonReader.h"
+#include "Serialization/JsonSerializer.h"
+#include "Serialization/JsonSerializerMacros.h"
+#include "Serialization/MemoryReader.h"
+#include "Serialization/MemoryWriter.h"
+
+// The following inlcudes are all for the Utility namespace
+#if PLATFORM_WINDOWS || PLATFORM_HOLOLENS
+	#include "Windows/WindowsHWrapper.h"
+	#include "Windows/AllowWindowsPlatformTypes.h"
+#endif
+
+#include "curl/curl.h"
+
+#if PLATFORM_WINDOWS || PLATFORM_HOLOLENS
+	#include "Windows/HideWindowsPlatformTypes.h"
+#endif
+
+#include "Containers/StaticArray.h"
+#include "Dom/JsonObject.h"
+#include "HAL/PlatformProcess.h"
+#include "Misc/ScopeRWLock.h"
+#include "Misc/SecureHash.h"
+#include "ProfilingDebugging/CountersTrace.h"
+
+#if WITH_SSL
+	#include "Ssl.h"
+#define UI UI_ST // Work around for name clash with the UI namespace in CoreUObject
+	#include <openssl/ssl.h>
+#undef UI
+#endif //WITH_SSL
+
+/** 
+ * The min version that the HordeStorage service must be in order for us to connect.
+ * Only increase the minimum required version to ensure that specific features
+ * are present that the code cannot run without.
+ */
+#define UE_HORDESTORAGE_MIN_MAJOR_VER 0
+#define UE_HORDESTORAGE_MIN_MINOR_VER 27
+#define UE_HORDESTORAGE_MIN_PATCH_VER 5
+
+/** When enabled we will only attempt to upload a payload if HordeStorage claims to not already have it.
+	Disabling this is only intended for debug purposes. */
+#define UE_CHECK_FOR_EXISTING_PAYLOADS 1
+
+/** When enabled we will only attempt to upload a chunk if HordeStorage claims to not already have it. 
+	In practice this doesn't really speed up the workflow hence being disabled. */
+#define UE_CHECK_FOR_EXISTING_CHUNKS 0
+
+/** When enabled we will attempt to PUT/GET many chunks at the same time asynchronously to improve throughput.
+	Disabling this is only intended for debug purposes as it is significantly slower. */
+#define UE_ENABLE_ASYNC_CHUNK_ACCESS 1
+
+TRACE_DECLARE_INT_COUNTER(HordeStorage_WaitOnRequestPool, TEXT("HordeStorage Wait On Request Pool"));
+
+namespace UE
+{
+// Code in the following name space has been taken from HttpDerivedDataBackend.cpp
+// once done we should consider what functionality is used by both implementations 
+// and what might be used by future implementations then decide what should be moved 
+// to a shared module
+namespace Utility
+{
+#define UE_MIRAGE_REQUEST_POOL_WAIT_INTERVAL 0.01f
+#define UE_MIRAGE_REQUEST_TIMEOUT_SECONDS 30L
+#define UE_MIRAGE_REQUEST_TIMOUT_ENABLED 1
+#define UE_MIRAGE_DEBUG 0
+#define UE_MIRAGE_REQUEST_POOL_SIZE 64
+#define UE_MIRAGE_MAX_FAILED_LOGIN_ATTEMPTS 16
+#define UE_MIRAGE_MAX_ATTEMPTS 4
+#define UE_MIRAGE_MAX_BUFFER_RESERVE 104857600u
+
+#if WITH_SSL
+static CURLcode sslctx_function(CURL* curl, void* sslctx, void* parm);
+#endif //WITH_SSL
+
+/**
+* Encapsulation for access token shared by all requests.
+*/
+struct FAccessToken
+{
+public:
+	FAccessToken() = default;
+	FString GetHeader() const;
+	void SetHeader(const TCHAR*);
+	uint32 GetSerial() const;
+
+private:
+	mutable FRWLock	Lock;
+	FString			Token;
+	uint32			Serial;
+};
+
+FString FAccessToken::GetHeader() const
+{
+	FReadScopeLock _(Lock);
+
+	return FString::Printf(TEXT("Authorization: Bearer %s"), *Token);
+}
+
+void FAccessToken::SetHeader(const TCHAR* InToken)
+{
+	FWriteScopeLock _(Lock);
+
+	Token = InToken;
+	Serial++;
+}
+
+uint32 FAccessToken::GetSerial() const
+{
+	return Serial;
+}
+
+/**
+* Minimal HTTP request type wrapping CURL without the need for managers. This request
+* is written to allow reuse of request objects, in order to allow connections to be reused.
+*
+* CURL has a global library initialization (curl_global_init). We rely on this happening in 
+* the Online/HTTP library which is a dependency on this module.
+*/
+class FRequest
+{
+public:
+	/**
+	* Supported request verbs
+	*/
+	enum RequestVerb
+	{
+		Get,
+		Put,
+		PutJson,
+		Post,
+		PostJson,
+		Delete,
+		Head
+	};
+
+	/**
+	* Convenience result type interpreted from HTTP response code.
+	*/
+	enum Result
+	{
+		Success,
+		Failed
+	};
+
+	enum EContentType
+	{
+		None = 0,
+		OctetStream,
+		UrlEncoded,
+		Json,
+		Xml
+	};
+
+	FRequest(const TCHAR* InDomain, FAccessToken* InAuthorizationToken, bool bInLogErrors)
+		: bLogErrors(bInLogErrors)
+		, Domain(InDomain)
+		, AuthorizationToken(InAuthorizationToken)
+	{
+		Curl = curl_easy_init();
+		Reset();
+	}
+
+	~FRequest()
+	{
+		curl_easy_cleanup(Curl);
+	}
+
+	/**
+	* Resets all options on the request except those that should always be set.
+	*/
+	void Reset()
+	{
+		Headers.Reset();
+		ResponseHeader.Reset();
+		ResponseBuffer.Reset();
+		ResponseCode = 0;
+		ReadDataView = TArrayView<const uint8>();
+		WriteDataView.Reset();
+		WriteDataBufferPtr = nullptr;
+		WriteHeaderBufferPtr = nullptr;
+		BytesSent = 0;
+		BytesReceived = 0;
+		CurlResult = CURL_LAST;
+
+		curl_easy_reset(Curl);
+
+		// Options that are always set for all connections.
+#if UE_MIRAGE_REQUEST_TIMOUT_ENABLED
+		curl_easy_setopt(Curl, CURLOPT_CONNECTTIMEOUT, UE_MIRAGE_REQUEST_TIMEOUT_SECONDS);
+#endif // UE_MIRAGE_REQUEST_TIMOUT_ENABLED
+		curl_easy_setopt(Curl, CURLOPT_FOLLOWLOCATION, 1L);
+		curl_easy_setopt(Curl, CURLOPT_NOSIGNAL, 1L);
+		// SSL options
+		curl_easy_setopt(Curl, CURLOPT_USE_SSL, CURLUSESSL_ALL);
+		curl_easy_setopt(Curl, CURLOPT_SSL_VERIFYPEER, 1);
+		curl_easy_setopt(Curl, CURLOPT_SSL_VERIFYHOST, 1);
+		curl_easy_setopt(Curl, CURLOPT_SSLCERTTYPE, "PEM");
+		// Response functions
+		curl_easy_setopt(Curl, CURLOPT_HEADERDATA, this);
+		curl_easy_setopt(Curl, CURLOPT_HEADERFUNCTION, &FRequest::StaticWriteHeaderFn);
+		curl_easy_setopt(Curl, CURLOPT_WRITEDATA, this);
+		curl_easy_setopt(Curl, CURLOPT_WRITEFUNCTION, StaticWriteBodyFn);
+		// SSL certification verification
+		curl_easy_setopt(Curl, CURLOPT_CAINFO, nullptr);
+		// TODO: What do we do if we do not have SSL?
+#if WITH_SSL
+		curl_easy_setopt(Curl, CURLOPT_SSL_CTX_FUNCTION, *sslctx_function);
+#endif // WITH_SSL
+		curl_easy_setopt(Curl, CURLOPT_SSL_CTX_DATA, this);
+		// Allow compressed data
+		curl_easy_setopt(Curl, CURLOPT_ACCEPT_ENCODING, "gzip");
+		// Rewind method, handle special error case where request need to rewind data stream
+		curl_easy_setopt(Curl, CURLOPT_SEEKFUNCTION, StaticSeekFn);
+		curl_easy_setopt(Curl, CURLOPT_SEEKDATA, this);
+		// Debug hooks
+#if UE_MIRAGE_DEBUG
+		curl_easy_setopt(Curl, CURLOPT_DEBUGDATA, this);
+		curl_easy_setopt(Curl, CURLOPT_DEBUGFUNCTION, StaticDebugCallback);
+		curl_easy_setopt(Curl, CURLOPT_VERBOSE, 1L);
+#endif // UE_MIRAGE_DEBUG
+	}
+
+	/** Gets the domain name for this request */
+	const FString& GetDomain() const
+	{
+		return Domain;
+	}
+
+	/** Returns the HTTP response code.*/
+	const int64 GetResponseCode() const
+	{
+		return ResponseCode;
+	}
+
+	/** Returns the number of bytes received this request (headers withstanding). */
+	const size_t GetBytesReceived() const
+	{
+		return BytesReceived;
+	}
+
+	/** Returns the number of bytes sent during this request (headers withstanding). */
+	const size_t  GetBytesSent() const
+	{
+		return BytesSent;
+	}
+
+	/**
+	* Upload buffer using the request, using either "Put" or "Post" verbs.
+	* @param Uri Url to use.
+	* @param Buffer Data to upload
+	* @return Result of the request
+	*/
+	template<RequestVerb V>
+	Result PerformBlockingUpload(const TCHAR* Uri, TArrayView<const uint8> Buffer)
+	{
+		static_assert(V == Put || V == PutJson || V == Post || V == PostJson, "Upload should use either Put or Post verbs.");
+
+		uint32 ContentLength = 0u;
+		EContentType ContentType = EContentType::None;
+
+		// TODO: Clean up the conditions here, too much code duplication
+		if (V == Put|| V == PutJson)
+		{
+			curl_easy_setopt(Curl, CURLOPT_UPLOAD, 1L);
+			curl_easy_setopt(Curl, CURLOPT_INFILESIZE, Buffer.Num());
+			curl_easy_setopt(Curl, CURLOPT_READDATA, this);
+			curl_easy_setopt(Curl, CURLOPT_READFUNCTION, StaticReadFn);
+
+			ContentType = V == Put ? EContentType::OctetStream : EContentType::Json;
+			ContentLength = Buffer.Num();
+			ReadDataView = Buffer;
+		}
+		else if (V == Post || V == PostJson)
+		{
+			curl_easy_setopt(Curl, CURLOPT_POST, 1L);
+			curl_easy_setopt(Curl, CURLOPT_INFILESIZE, Buffer.Num());
+			curl_easy_setopt(Curl, CURLOPT_READDATA, this);
+			curl_easy_setopt(Curl, CURLOPT_READFUNCTION, StaticReadFn);
+
+			ContentType = V == Post ? EContentType::UrlEncoded : EContentType::Json;
+			ContentLength = Buffer.Num();
+			ReadDataView = Buffer;
+		}
+
+		return PerformBlocking(Uri, V, ContentType, ContentLength);
+	}
+
+	/**
+	* Download an url into a buffer using the request.
+	* @param Uri Url to use.
+	* @param Buffer Optional buffer where data should be downloaded to. If empty downloaded data will
+	* be stored in an internal buffer and accessed GetResponse* methods.
+	* @return Result of the request
+	*/
+	Result PerformBlockingDownload(const TCHAR* Uri)
+	{
+		curl_easy_setopt(Curl, CURLOPT_HTTPGET, 1L);
+
+		WriteDataBufferPtr = nullptr;
+		WriteDataView.Reset();
+
+		return PerformBlocking(Uri, Get, EContentType::None, 0u);
+	}
+
+	Result PerformBlockingDownload(const TCHAR* Uri, TArray<uint8>* Buffer)
+	{
+		curl_easy_setopt(Curl, CURLOPT_HTTPGET, 1L);
+		
+		WriteDataBufferPtr = Buffer;
+		WriteDataView.Reset();
+
+		return PerformBlocking(Uri, Get, EContentType::None, 0u);
+	}
+
+	Result PerformBlockingDownload(const TCHAR* Uri, FMutableMemoryView Buffer)
+	{
+		curl_easy_setopt(Curl, CURLOPT_HTTPGET, 1L);
+
+		WriteDataBufferPtr = nullptr;
+		WriteDataView = Buffer;
+
+		return PerformBlocking(Uri, Get, EContentType::None, 0u);
+	}	
+
+	/**
+	* Query an url using the request. Queries can use either "Head" or "Delete" verbs.
+	* @param Uri Url to use.
+	* @return Result of the request
+	*/
+	template<RequestVerb V>
+	Result PerformBlockingQuery(const TCHAR* Uri)
+	{
+		static_assert(V == Head || V == Delete, "Queries should use either Head or Delete verbs.");
+
+		if (V == Delete)
+		{
+			curl_easy_setopt(Curl, CURLOPT_CUSTOMREQUEST, "DELETE");
+		}
+		else if (V == Head)
+		{
+			curl_easy_setopt(Curl, CURLOPT_NOBODY, 1L);
+		}
+
+		return PerformBlocking(Uri, V, EContentType::None, 0u);
+	}
+
+	/**
+	* Set a header to send with the request.
+	*/
+	void SetHeader(const TCHAR* Header, const TCHAR* Value)
+	{
+		check(CurlResult == CURL_LAST); // Cannot set header after request is sent
+		Headers.Add(FString::Printf(TEXT("%s: %s"), Header, Value));
+	}
+
+	/**
+	* Attempts to find the header from the response. Returns false if header is not present.
+	*/
+	bool GetHeader(const ANSICHAR* Header, FString& OutValue) const
+	{
+		check(CurlResult != CURL_LAST);  // Cannot query headers before request is sent
+
+		const ANSICHAR* HeadersBuffer = (const ANSICHAR*)ResponseHeader.GetData();
+		size_t HeaderLen = strlen(Header);
+
+		// Find the header key in the (ANSI) response buffer. If not found we can exist immediately
+		if (const ANSICHAR* Found = strstr(HeadersBuffer, Header))
+		{
+			const ANSICHAR* Linebreak = strchr(Found, '\r');
+			const ANSICHAR* ValueStart = Found + HeaderLen + 2; //colon and space
+			const size_t ValueSize = Linebreak - ValueStart;
+			FUTF8ToTCHAR TCHARData(ValueStart, ValueSize);
+			OutValue = FString(TCHARData.Length(), TCHARData.Get());
+			return true;
+		}
+		return false;
+	}
+
+	/**
+	* Returns the response buffer. Note that is the request is performed
+	* with an external buffer as target buffer this string will be empty.
+	*/
+	const TArray<uint8>& GetResponseBuffer() const
+	{
+		return ResponseBuffer;
+	}
+
+	/**
+	* Returns the response buffer as a string. Note that is the request is performed
+	* with an external buffer as target buffer this string will be empty.
+	*/
+	FString GetResponseAsString() const
+	{
+		return GetAnsiBufferAsString(ResponseBuffer);
+	}
+
+	/**
+	* Returns the response header as a string.
+	*/
+	FString GetResponseHeaderAsString()
+	{
+		return GetAnsiBufferAsString(ResponseHeader);
+	}
+
+	/**
+	* Tries to parse the response buffer as a JsonObject. Return empty pointer if
+	* parse error occurs.
+	*/
+	TSharedPtr<FJsonObject> GetResponseAsJsonObject() const
+	{
+		FString Response = GetAnsiBufferAsString(ResponseBuffer);
+
+		TSharedPtr<FJsonObject> JsonObject;
+		TSharedRef<TJsonReader<> > JsonReader = TJsonReaderFactory<>::Create(Response);
+		if (!FJsonSerializer::Deserialize(JsonReader, JsonObject) || !JsonObject.IsValid())
+		{
+			return TSharedPtr<FJsonObject>(nullptr);
+		}
+
+		return JsonObject;
+	}
+
+	/** Will return true if the response code is considered a success */
+	static bool IsSuccessfulResponse(int64 ResponseCode)
+	{
+		// We consider anything in the 1XX or 2XX range a success
+		return ResponseCode >= 100 && ResponseCode < 300;
+	}
+
+private:
+
+	CURL* Curl;
+	CURLcode				CurlResult;
+	long					ResponseCode;
+	size_t					BytesSent;
+	size_t					BytesReceived;
+	bool					bLogErrors;
+
+	TArrayView<const uint8>	ReadDataView;
+
+	FMutableMemoryView		WriteDataView;
+	TArray<uint8>*			WriteDataBufferPtr;
+	TArray<uint8>*			WriteHeaderBufferPtr;
+
+	TArray<uint8>			ResponseHeader;
+	TArray<uint8>			ResponseBuffer;
+	TArray<FString>			Headers;
+	FString					Domain;
+	FAccessToken* AuthorizationToken;
+
+	/**
+	* Performs the request, blocking until finished.
+	* @param Uri Address on the domain to query
+	* @param Verb HTTP verb to use
+	* @param Buffer Optional buffer to directly receive the result of the request.
+	* If unset the response body will be stored in the request.
+	*/
+	Result PerformBlocking(const TCHAR* Uri, RequestVerb Verb, EContentType ContentType, uint32 ContentLength)
+	{
+		static const char* CommonHeaders[] = {
+			"User-Agent: UE4",
+			nullptr
+		};
+
+		TRACE_CPUPROFILER_EVENT_SCOPE(FRequest::PerformBlocking);
+
+		// Reset a few values incase the request has been reused (re-submitting a failure etc)
+		BytesSent = 0;
+		BytesReceived = 0;
+
+		// Setup request options
+		TStringBuilder<256> Url;
+		Url << Domain << TEXT("/") << Uri;
+
+		curl_easy_setopt(Curl, CURLOPT_URL, TCHAR_TO_ANSI(*Url));
+
+		// Setup response header buffer. If caller has not setup a response data buffer, use internal.
+		WriteHeaderBufferPtr = &ResponseHeader;
+		if (WriteDataBufferPtr == nullptr && WriteDataView.IsEmpty())
+		{
+			WriteDataBufferPtr = &ResponseBuffer;
+		}
+
+		// Build headers list
+		curl_slist* CurlHeaders = nullptr;
+		// Add common headers
+		for (uint8 i = 0; CommonHeaders[i] != nullptr; ++i)
+		{
+			CurlHeaders = curl_slist_append(CurlHeaders, CommonHeaders[i]);
+		}
+		
+		// Content-Length should always be set
+		TAnsiStringBuilder<48> ContentLengthString;
+		ContentLengthString.Appendf("Content-Length: %d", ContentLength);
+		CurlHeaders = curl_slist_append(CurlHeaders, *ContentLengthString);
+
+		// TODO: Clean up
+		switch (ContentType)
+		{
+		case OctetStream:
+			CurlHeaders = curl_slist_append(CurlHeaders, "Content-Type: application/octet-stream");
+			break;
+		case UrlEncoded:
+			CurlHeaders = curl_slist_append(CurlHeaders, "Content-Type: application/x-www-form-urlencoded");
+			break;
+		case Json:
+			CurlHeaders = curl_slist_append(CurlHeaders, "Content-Type: application/json");
+			break;
+		case Xml:
+			CurlHeaders = curl_slist_append(CurlHeaders, "Content-Type: application/xml");
+			break;
+		}
+
+		// Setup added headers
+		for (const FString& Header : Headers)
+		{
+			CurlHeaders = curl_slist_append(CurlHeaders, TCHAR_TO_ANSI(*Header));
+		}
+
+		// And auth token if it's set
+		if (AuthorizationToken != nullptr)
+		{
+			CurlHeaders = curl_slist_append(CurlHeaders, TCHAR_TO_ANSI(*AuthorizationToken->GetHeader()));
+		}
+
+		curl_easy_setopt(Curl, CURLOPT_HTTPHEADER, CurlHeaders);
+
+		// Shots fired!
+		CurlResult = curl_easy_perform(Curl);
+
+		// Get response code
+		bool bRedirected = false;
+		if (CURLE_OK == curl_easy_getinfo(Curl, CURLINFO_RESPONSE_CODE, &ResponseCode))
+		{
+			bRedirected = (ResponseCode >= 300 && ResponseCode < 400);
+		}
+
+		LogResult(CurlResult, Uri, Verb);
+
+		// Clean up
+		curl_slist_free_all(CurlHeaders);
+
+		return CurlResult == CURLE_OK ? Success : Failed;
+	}
+
+	void LogResult(CURLcode Result, const TCHAR* Uri, RequestVerb Verb) const
+	{
+		TRACE_CPUPROFILER_EVENT_SCOPE(FRequest::LogResult);
+
+		if (Result == CURLE_OK)
+		{
+			bool bSuccess = false;
+			const TCHAR* VerbStr = nullptr;
+			TStringBuilder<256> AdditionalInfo;
+
+			switch (Verb)
+			{
+			case Head:
+				// TODO: Io returns 404 if the head request is not found, Europa returns 400. Clean this up once 
+				// the inconsistency on the HordeStorage server is fixed.
+				bSuccess = (ResponseCode == 400 || ResponseCode == 404 || IsSuccessfulResponse(ResponseCode));
+				VerbStr = TEXT("querying");
+				break;
+			case Get:
+				bSuccess = (ResponseCode == 400 || IsSuccessfulResponse(ResponseCode));
+				VerbStr = TEXT("fetching");
+				AdditionalInfo.Appendf(TEXT("Received: %d bytes."), BytesReceived);
+				break;
+			case Put:
+			case PutJson:
+				bSuccess = IsSuccessfulResponse(ResponseCode);
+				VerbStr = TEXT("updating");
+				AdditionalInfo.Appendf(TEXT("Sent: %d bytes."), BytesSent);
+				break;
+			case Post:
+			case PostJson:
+				bSuccess = IsSuccessfulResponse(ResponseCode);
+				VerbStr = TEXT("posting");
+				break;
+			case Delete:
+				bSuccess = IsSuccessfulResponse(ResponseCode);
+				VerbStr = TEXT("deleting");
+				break;
+			default:
+				checkf(false, TEXT("Unknown RequestVerb found in FRequet::LogResult"));
+			}
+
+			if (bSuccess)
+			{
+				UE_LOG(
+					LogVirtualization,
+					VeryVerbose,
+					TEXT("Finished %s HTTP cache entry (response %d) from %s. %s"),
+					VerbStr,
+					ResponseCode,
+					Uri,
+					*AdditionalInfo
+				);
+			}
+			else if (bLogErrors)
+			{
+				// Print the response body if we got one, otherwise print header.
+				FString Response = GetAnsiBufferAsString(ResponseBuffer.Num() > 0 ? ResponseBuffer : ResponseHeader);
+				Response.ReplaceCharInline('\n', ' ');
+				Response.ReplaceCharInline('\r', ' ');
+				
+				// Don't log access denied as error, since tokens can expire mid session
+				if (ResponseCode == 401)
+				{
+					UE_LOG(
+						LogVirtualization,
+						VeryVerbose,
+						TEXT("Failed %s HTTP cache entry (response %d) from %s. Response: %s"),
+						VerbStr,
+						ResponseCode,
+						Uri,
+						*Response
+					);
+				}
+				else
+				{	
+					UE_LOG(
+						LogVirtualization,
+						Error,
+						TEXT("Failed %s HTTP cache entry (response %d) from %s. Response: %s"),
+						VerbStr,
+						ResponseCode,
+						Uri,
+						*Response
+					);
+				}
+			}
+		}
+		else if (bLogErrors)
+		{
+			UE_LOG(
+				LogVirtualization,
+				Error,
+				TEXT("Error while connecting to %s: %s"),
+				*Domain,
+				ANSI_TO_TCHAR(curl_easy_strerror(Result))
+			);
+		}
+	}
+
+	FString GetAnsiBufferAsString(const TArray<uint8>& Buffer) const
+	{
+		// Content is NOT null-terminated; we need to specify lengths here
+		FUTF8ToTCHAR TCHARData(reinterpret_cast<const ANSICHAR*>(Buffer.GetData()), Buffer.Num());
+		return FString(TCHARData.Length(), TCHARData.Get());
+	}
+
+#if UE_MIRAGE_DEBUG
+	static size_t StaticDebugCallback(CURL* Handle, curl_infotype DebugInfoType, char* DebugInfo, size_t DebugInfoSize, void* UserData)
+	{
+		FRequest* Request = static_cast<FRequest*>(UserData);
+
+		auto PrintText = [](FRequest* InRequest, char* InDebugInfo, size_t InDebugInfoSize)
+		{
+			// Truncate at 1023 characters. This is just an arbitrary number based on a buffer size seen in
+			// the libcurl code.
+			InDebugInfoSize = FMath::Min(InDebugInfoSize, (size_t)1023);
+
+			// Calculate the actual length of the string due to incorrect use of snprintf() in lib/vtls/openssl.c.
+			char* FoundNulPtr = (char*)memchr(InDebugInfo, 0, InDebugInfoSize);
+			int CalculatedSize = FoundNulPtr != nullptr ? FoundNulPtr - InDebugInfo : InDebugInfoSize;
+
+			auto ConvertedString = StringCast<TCHAR>(static_cast<const ANSICHAR*>(InDebugInfo), CalculatedSize);
+			FString DebugText(ConvertedString.Length(), ConvertedString.Get());
+
+			UE_LOG(LogVirtualization, Log, TEXT("%p: '%s'"), InRequest, *DebugText);
+		};
+
+		switch (DebugInfoType)
+		{
+		case CURLINFO_TEXT:
+		{
+			PrintText(Request, DebugInfo, DebugInfoSize);
+		}
+		break;
+
+		case CURLINFO_HEADER_IN:
+			UE_LOG(LogVirtualization, Log, TEXT("%p: Received header (%d bytes)"), Request, DebugInfoSize);
+			break;
+
+		case CURLINFO_HEADER_OUT:
+			UE_LOG(LogVirtualization, Log, TEXT("%p: Sent header (%d bytes)"), Request, DebugInfoSize);
+			PrintText(Request, DebugInfo, DebugInfoSize);
+			break;
+
+		case CURLINFO_DATA_IN:
+			UE_LOG(LogVirtualization, Log, TEXT("%p: Received data (%d bytes)"), Request, DebugInfoSize);
+			break;
+
+		case CURLINFO_DATA_OUT:
+			UE_LOG(LogVirtualization, Log, TEXT("%p: Sent data (%d bytes)"), Request, DebugInfoSize);
+			break;
+
+		case CURLINFO_SSL_DATA_IN:
+			UE_LOG(LogVirtualization, Log, TEXT("%p: Received SSL data (%d bytes)"), Request, DebugInfoSize);
+			break;
+
+		case CURLINFO_SSL_DATA_OUT:
+			UE_LOG(LogVirtualization, Log, TEXT("%p: Sent SSL data (%d bytes)"), Request, DebugInfoSize);
+			break;
+		}
+
+		return 0;
+	}
+#endif // UE_MIRAGE_DEBUG
+
+	static size_t StaticReadFn(void* Ptr, size_t SizeInBlocks, size_t BlockSizeInBytes, void* UserData)
+	{
+		TRACE_CPUPROFILER_EVENT_SCOPE(FRequest::StaticReadFn);
+		
+		FRequest* Request = static_cast<FRequest*>(UserData);
+		TArrayView<const uint8>& ReadDataView = Request->ReadDataView;
+
+		const size_t Offset = Request->BytesSent;
+		const size_t ReadSize = FMath::Min((size_t)ReadDataView.Num() - Offset, SizeInBlocks * BlockSizeInBytes);
+		check(ReadDataView.Num() >= Offset + ReadSize);
+
+		FMemory::Memcpy(Ptr, ReadDataView.GetData() + Offset, ReadSize);
+		Request->BytesSent += ReadSize;
+
+		return ReadSize;
+	}
+
+	static size_t StaticWriteHeaderFn(void* Ptr, size_t SizeInBlocks, size_t BlockSizeInBytes, void* UserData)
+	{
+		TRACE_CPUPROFILER_EVENT_SCOPE(FRequest::StaticWriteHeaderFn);
+
+		FRequest* Request = static_cast<FRequest*>(UserData);
+		const size_t WriteSize = SizeInBlocks * BlockSizeInBytes;
+		TArray<uint8>* WriteHeaderBufferPtr = Request->WriteHeaderBufferPtr;
+		if (WriteHeaderBufferPtr && WriteSize > 0)
+		{
+			const size_t CurrentBufferLength = WriteHeaderBufferPtr->Num();
+			if (CurrentBufferLength > 0)
+			{
+				// Remove the previous zero termination
+				(*WriteHeaderBufferPtr)[CurrentBufferLength - 1] = ' ';
+			}
+
+			// Write the header
+			WriteHeaderBufferPtr->Append((const uint8*)Ptr, WriteSize + 1);
+			(*WriteHeaderBufferPtr)[WriteHeaderBufferPtr->Num() - 1] = 0; // Zero terminate string
+			return WriteSize;
+		}
+		return 0;
+	}
+
+	static size_t StaticWriteBodyFn(void* Ptr, size_t SizeInBlocks, size_t BlockSizeInBytes, void* UserData)
+	{
+		TRACE_CPUPROFILER_EVENT_SCOPE(FRequest::StaticWriteBodyFn);
+
+		FRequest* Request = static_cast<FRequest*>(UserData);
+		const size_t WriteSize = SizeInBlocks * BlockSizeInBytes;
+
+		if (WriteSize == 0)
+		{
+			return 0;
+		}
+
+		if (TArray<uint8>* WriteDataBufferPtr = Request->WriteDataBufferPtr)
+		{
+			// If this is the first part of the body being received, try to reserve 
+			// memory if content length is defined in the header.
+			if (Request->BytesReceived == 0 && Request->WriteHeaderBufferPtr)
+			{
+				static const ANSICHAR* ContentLengthHeaderStr = "Content-Length: ";
+				const ANSICHAR* Header = (const ANSICHAR*)Request->WriteHeaderBufferPtr->GetData();
+
+				if (const ANSICHAR* ContentLengthHeader = FCStringAnsi::Strstr(Header, ContentLengthHeaderStr))
+				{
+					size_t ContentLength = (size_t)FCStringAnsi::Atoi64(ContentLengthHeader + strlen(ContentLengthHeaderStr));
+					if (ContentLength > 0u && ContentLength < UE_MIRAGE_MAX_BUFFER_RESERVE)
+					{
+						WriteDataBufferPtr->Reserve(ContentLength);
+					}
+				}
+			}
+
+			// Write to the target buffer
+			WriteDataBufferPtr->Append((const uint8*)Ptr, WriteSize);
+			Request->BytesReceived += WriteSize;
+			return WriteSize;
+		}
+		else if (!Request->WriteDataView.IsEmpty())
+		{
+			if (Request->BytesReceived + WriteSize <= Request->WriteDataView.GetSize())
+			{
+				FMemory::Memcpy((uint8*)Request->WriteDataView.GetData() + Request->BytesReceived, Ptr, WriteSize);
+				Request->BytesReceived += WriteSize;
+
+				return WriteSize;
+			}
+			else
+			{
+				UE_CLOG(
+					Request->bLogErrors,
+					LogVirtualization,
+					Error,
+					TEXT("Attempting to write %d bytes to the response buffer which only has %" SSIZE_T_FMT " bytes remaining %" SSIZE_T_FMT),
+					WriteSize,
+					(size_t)Request->WriteDataView.GetSize() - Request->BytesReceived
+				);
+
+				return -1;
+			}
+		}
+		else
+		{
+			UE_CLOG(Request->bLogErrors, LogVirtualization, Error, TEXT("No response buffer was set!"));	
+			return -1;
+		}
+	}
+
+	static size_t StaticSeekFn(void* UserData, curl_off_t Offset, int Origin)
+	{
+		TRACE_CPUPROFILER_EVENT_SCOPE(FRequest::StaticSeekFn);
+
+		FRequest* Request = static_cast<FRequest*>(UserData);
+		size_t NewPosition = 0;
+
+		switch (Origin)
+		{
+		case SEEK_SET: NewPosition = Offset; break;
+		case SEEK_CUR: NewPosition = Request->BytesSent + Offset; break;
+		case SEEK_END: NewPosition = Request->ReadDataView.Num() + Offset; break;
+		}
+
+		// Make sure we don't seek outside of the buffer
+		if (NewPosition < 0 || NewPosition >= Request->ReadDataView.Num())
+		{
+			return CURL_SEEKFUNC_FAIL;
+		}
+
+		// Update the used offset
+		Request->BytesSent = NewPosition;
+		return CURL_SEEKFUNC_OK;
+	}
+};
+
+/**
+* Pool that manages a fixed set of requests. Users are required to release requests that have been 
+* acquired. Usable with \ref FScopedRequestPtr which handles this automatically.
+*/
+struct FRequestPool
+{
+	FRequestPool(const TCHAR* InServiceUrl, FAccessToken* InAuthorizationToken)
+	{
+		for (int32 Index = 0; Index < Pool.Num(); ++Index)
+		{
+			Pool[Index].Usage = 0u;
+			Pool[Index].Request = new FRequest(InServiceUrl, InAuthorizationToken, true);	
+		}
+	}
+
+	~FRequestPool()
+	{
+		for (int32 Index = 0; Index < Pool.Num(); ++Index)
+		{
+			// No requests should be in use by now.
+			check(Pool[Index].Usage.Load(EMemoryOrder::Relaxed) == 0u);
+			delete Pool[Index].Request;
+		}
+	}
+
+	/**
+	* Block until a request is free. Once a request has been returned it is 
+	* "owned by the caller and need to release it to the pool when work has been completed.
+	* @return Usable request instance.
+	*/
+	FRequest* WaitForFreeRequest()
+	{
+		TRACE_CPUPROFILER_EVENT_SCOPE(FRequestPool::WaitForFreeRequest);
+		while (true)
+		{
+			for (int32 Index = 0; Index < Pool.Num(); ++Index)
+			{
+				if (!Pool[Index].Usage.Load(EMemoryOrder::Relaxed))
+				{
+					uint8 Expected = 0u;
+					if (Pool[Index].Usage.CompareExchange(Expected, 1u))
+					{
+						return Pool[Index].Request;
+					}
+				}
+			}
+
+			TRACE_COUNTER_ADD(HordeStorage_WaitOnRequestPool, int64(1));
+			FPlatformProcess::Sleep(UE_MIRAGE_REQUEST_POOL_WAIT_INTERVAL);
+		}
+	}
+
+	/**
+	* Re lease request to the pool.
+	* @param Request Request that should be freed. Note that any buffer owned by the request can now be reset.
+	*/
+	void ReleaseRequestToPool(FRequest* Request)
+	{
+		for (int32 Index = 0; Index < Pool.Num(); ++Index)
+		{
+			if (Pool[Index].Request == Request)
+			{
+				Request->Reset();
+				Pool[Index].Usage.Exchange(0u);
+				return;
+			}
+		}
+		check(false);
+	}
+
+private:
+
+	struct FEntry
+	{
+		TAtomic<uint8> Usage;
+		FRequest* Request;
+	};
+
+	TStaticArray<FEntry, UE_MIRAGE_REQUEST_POOL_SIZE> Pool;
+
+	FRequestPool() {}
+};
+
+/**
+* Utility class to manage requesting and releasing requests from the \ref FRequestPool.
+*/
+struct FScopedRequestPtr
+{
+public:
+	FScopedRequestPtr(FRequestPool* InPool)
+		: Request(InPool->WaitForFreeRequest())
+		, Pool(InPool)
+	{}
+
+	~FScopedRequestPtr()
+	{
+		Reset();
+	}
+
+	bool IsValid() const 
+	{
+		return Request != nullptr;
+	}
+
+	FRequest* Get() const
+	{
+		check(IsValid());
+		return Request;
+	}
+
+	FRequest* operator->()
+	{
+		check(IsValid());
+		return Request;
+	}
+
+	void Reset()
+	{
+		if (IsValid())
+		{
+			Pool->ReleaseRequestToPool(Request);
+			Request = nullptr;
+		}
+	}
+
+private:
+	FRequest* Request;
+	FRequestPool* Pool;
+};
+
+
+/**
+* Adds a checksum (as request header) for a given payload. HordeStorage will use this to verify the integrity
+* of the received data.
+* @param Request Request that the data will be sent with.
+* @param Payload Payload that will be sent.
+* @return True on success, false on failure.
+*/
+FSHAHash HashPayload(FRequest* Request, const TArrayView<const uint8> Payload)
+{
+	FSHAHash PayloadHash;
+	FSHA1::HashBuffer(Payload.GetData(), Payload.Num(), PayloadHash.Hash);
+	Request->SetHeader(TEXT("X-Jupiter-Sha1"), *PayloadHash.ToString());
+	return PayloadHash;
+}
+
+/**
+* Verifies the integrity of the received data using supplied checksum.
+* @param Hash Received hash value.
+* @param Payload Payload received.
+* @return True if the data is correct, false if checksums doesn't match.
+*/
+bool VerifyPayload(FSHAHash Hash, const TArray<uint8>& Payload)
+{
+	FSHAHash PayloadHash;
+	FSHA1::HashBuffer(Payload.GetData(), Payload.Num(), PayloadHash.Hash);
+
+	if (Hash != PayloadHash)
+	{
+		UE_LOG(LogVirtualization,
+			Warning,
+			TEXT("Checksum from server did not match recieved data (%s vs %s). Discarding cached result."),
+			*Hash.ToString(),
+			*PayloadHash.ToString()
+		);
+		return false;
+	}
+
+	return true;
+}
+
+/**
+* Verifies the integrity of the received data using supplied checksum.
+* @param Request Request that the data was be received with.
+* @param Payload Payload received.
+* @return True if the data is correct, false if checksums doesn't match.
+*/
+bool VerifyRequest(const FRequest* Request, const TArray<uint8>& Payload)
+{
+	FString RecievedHashStr;
+	FSHAHash PayloadHash;
+	FSHA1::HashBuffer(Payload.GetData(), Payload.Num(), PayloadHash.Hash);
+	if (Request->GetHeader("X-Jupiter-Sha1", RecievedHashStr))
+	{
+		FSHAHash RecievedHash;
+		RecievedHash.FromString(RecievedHashStr);
+		return VerifyPayload(RecievedHash, Payload);
+	}
+	UE_LOG(LogVirtualization, Error, TEXT("HTTP server did not send a content hash. Wrong server version?"));
+	return true;
+}
+
+#if WITH_SSL
+static int SslCertVerify(int PreverifyOk, X509_STORE_CTX* Context)
+{
+	if (PreverifyOk == 1)
+	{
+		SSL* Handle = static_cast<SSL*>(X509_STORE_CTX_get_ex_data(Context, SSL_get_ex_data_X509_STORE_CTX_idx()));
+		check(Handle);
+
+		SSL_CTX* SslContext = SSL_get_SSL_CTX(Handle);
+		check(SslContext);
+
+		FRequest* Request = static_cast<FRequest*>(SSL_CTX_get_app_data(SslContext));
+		check(Request);
+
+		const FString& Domain = Request->GetDomain();
+
+		if (!FSslModule::Get().GetCertificateManager().VerifySslCertificates(Context, Domain))
+		{
+			PreverifyOk = 0;
+		}
+	}
+
+	return PreverifyOk;
+}
+
+static CURLcode sslctx_function(CURL* curl, void* sslctx, void* parm)
+{
+	SSL_CTX* Context = static_cast<SSL_CTX*>(sslctx);
+	const ISslCertificateManager& CertificateManager = FSslModule::Get().GetCertificateManager();
+
+	CertificateManager.AddCertificatesToSslContext(Context);
+	SSL_CTX_set_verify(Context, SSL_CTX_get_verify_mode(Context), SslCertVerify);
+	SSL_CTX_set_app_data(Context, parm);
+
+	/* all set to go */
+	return CURLE_OK;
+}
+#endif //WITH_SSL
+
+} // namespace Utility
+
+namespace Virtualization
+{
+
+/** Represents the data required to make a PUT request to the Europa DDCCache API
+	that can be easily serialized to JSON to make the actual request. */
+struct FEuropaDDCCachePUTRequest : public FJsonSerializable
+{
+	struct FMetaData : public FJsonSerializable
+	{
+		BEGIN_JSON_SERIALIZER
+			JSON_SERIALIZE("payloadLength", PayloadLength);
+			JSON_SERIALIZE("chunkLength", ChunkLength);
+		END_JSON_SERIALIZER
+
+		/** Overall length (in bytes) of the payload */
+		int64 PayloadLength { INDEX_NONE };
+		/** The max length (in bytes) of each chunk */
+		int64 ChunkLength { INDEX_NONE };
+	};
+
+	BEGIN_JSON_SERIALIZER
+		JSON_SERIALIZE_ARRAY("blobReferences", ChunkHashes);
+		JSON_SERIALIZE_OBJECT_SERIALIZABLE("metadata", MetaData);
+		JSON_SERIALIZE("contentHash", PayloadHash);
+	END_JSON_SERIALIZER
+
+	/** A FSHAHash of the entire payload once it is reconstructed */
+	FString PayloadHash;
+	/** An object containing additional metadata */
+	FMetaData MetaData;
+	/** Each string represents a FSHAHash */
+	TArray<FString> ChunkHashes;
+};
+
+/** Data structure for use with HordeStorage Europa DDCCache GET JSON response */
+struct FDDCCacheGETResponse : public FJsonSerializable
+{
+	struct FMetaData : public FJsonSerializable
+	{
+		BEGIN_JSON_SERIALIZER
+			JSON_SERIALIZE("payloadLength", PayloadLength);
+			JSON_SERIALIZE("chunkLength", ChunkLength);
+		END_JSON_SERIALIZER
+
+		/** Overall length of the payload */
+		int64 PayloadLength { INDEX_NONE };
+		/** The max length of each chunk */
+		int64 ChunkLength { INDEX_NONE };
+	};
+
+	BEGIN_JSON_SERIALIZER
+		JSON_SERIALIZE("name", Name);
+		JSON_SERIALIZE("lastAccessTime", LastAccessTime);
+		JSON_SERIALIZE_OBJECT_SERIALIZABLE("metadata", MetaData);
+		JSON_SERIALIZE("contentHash", PayloadHash);
+		JSON_SERIALIZE_ARRAY("blobIdentifiers", ChunkHashes);
+		JSON_SERIALIZE("blob", PayloadBlob);
+	END_JSON_SERIALIZER
+
+	/** The payload name in the format: '{namespace}.{bucket}.{key}' */
+	FString Name;
+	/** The data and time that the payload was last accessed */
+	FDateTime LastAccessTime { INDEX_NONE };
+	/** A FSHAHash of the entire payload once it is reconstructed */
+	FString PayloadHash;
+	/** An object containing additional metadata */
+	FMetaData MetaData;
+	/** Each string represents a FSHAHash */
+	TArray<FString> ChunkHashes;
+	/** The payload (Base64 encoding) */
+	FString PayloadBlob;
+};
+
+/** Data structure for use with HordeStorage static GET response. Represents the status of the HordeStorage service. */
+struct FHttpServiceStatus : public FJsonSerializable
+{
+	/** Returns true if the current version in the object is greater or equal to the version numbers passed in */
+	bool DoesHaveValidVersion(uint32 MinMajorVersion, uint32 MinMinorVersion, uint32 MinPatchVersion) const 
+	{
+		uint32 MajorVersion = 0;
+		uint32 MinorVersion = 0;
+		uint32 PatchVersion = 0;
+
+		if (GetVersionNumbers(MajorVersion, MinorVersion, PatchVersion))
+		{
+			if (IsValidVersion(MinMajorVersion, MinMinorVersion, MinPatchVersion, MajorVersion, MinorVersion, PatchVersion))
+			{
+				return true;
+			}
+			else
+			{
+				UE_LOG(LogVirtualization, Error, TEXT("HordeStorage service version is too old! Found: '%s' Required: %u.%u.%u"), 
+					*Version, MajorVersion, MinorVersion, PatchVersion);
+				return false;
+			}
+		}
+		else
+		{
+			UE_LOG(LogVirtualization, Error, TEXT("Failed to parse valid HordeStorage version number from '%s'"), *Version);
+			return false;
+		}
+	}
+
+	/** Returns the version numbers parsed from string format into uint32 values, returns false if the string cannot be parsed */
+	bool GetVersionNumbers(uint32& OutMajorVersion, uint32& OutMinorVersion, uint32& OutPatchVersion) const
+	{
+		TCHAR* End = nullptr;
+		// Read the major/minor/patch numbers
+		uint64 MajorVersion = FCString::Strtoui64(*Version, &End, 10);
+		if (MajorVersion > MAX_uint32 || *(End++) != '.')
+		{
+			return false;
+		}
+
+		uint64 MinorVersion = FCString::Strtoui64(End, &End, 10);
+		if (MinorVersion > MAX_uint32 || *(End++) != '.')
+		{
+			return false;
+		}
+
+		uint64 PatchVersion = FCString::Strtoui64(End, &End, 10);
+		if (PatchVersion > MAX_uint32)
+		{
+			return false;
+		}
+
+		OutMajorVersion = static_cast<uint16>(MajorVersion);
+		OutMinorVersion = static_cast<uint16>(MinorVersion);
+		OutPatchVersion = static_cast<uint16>(PatchVersion);
+
+		return true;
+	}
+
+	/** Returns true if the given capability is supported by HordeStorage */
+	bool SupportsCapability(FStringView Capability) const 
+	{
+		return Capabilities.Contains(Capability);
+	}
+
+	/** Prints the status of the HordeStorage service to the log */
+	void LogStatusInfo() const
+	{
+		UE_LOG(LogVirtualization, Log, TEXT("HordeStorage Status:"));
+		UE_LOG(LogVirtualization, Log, TEXT("Version: %s"), *Version);
+		UE_LOG(LogVirtualization, Log, TEXT("Site Id: %s"), *SiteIdentifier);
+		UE_LOG(LogVirtualization, Log, TEXT("GitHash: %s"), *GitHash);		
+		UE_LOG(LogVirtualization, Log, TEXT("Capabilities:"));
+		for (const FString& Capability : Capabilities)
+		{
+			UE_LOG(LogVirtualization, Log, TEXT("\t%s"), *Capability);
+		}	
+	}
+
+private:
+
+	bool IsValidVersion(uint32 MinMajorVersion, uint32 MinMinorVersion, uint32 MinPatchVersion,
+						uint32 CurrentMajorVersion, uint32 CurrentMinorVersion, uint32 CurrentPatchVersion) const
+	{
+		if (CurrentMajorVersion != MinMajorVersion)
+		{
+			return CurrentMajorVersion > MinMajorVersion;
+		}
+
+		if (CurrentMinorVersion != MinMinorVersion)
+		{
+			return CurrentMinorVersion > MinMinorVersion;
+		}
+
+		if (CurrentPatchVersion != MinPatchVersion)
+		{
+			return CurrentPatchVersion > MinPatchVersion;
+		}
+
+		return true;
+	}
+
+	BEGIN_JSON_SERIALIZER
+		JSON_SERIALIZE("version", Version);
+		JSON_SERIALIZE("gitHash", GitHash);
+		JSON_SERIALIZE_ARRAY("capabilities", Capabilities);
+		JSON_SERIALIZE("siteIdentifier", SiteIdentifier);
+	END_JSON_SERIALIZER
+
+	/** Version of the service in the format [MAJOR.MINOR.PATCH] */
+	FString Version;
+	/**The git commit hash for HordeStorage */
+	FString GitHash;
+	/** An array of which sub-services HordeStorage supports */
+	TArray<FString> Capabilities;
+	/** The identifier for the server connected to */
+	FString SiteIdentifier;
+};
+
+FHttpBackend::FHttpBackend(FStringView ConfigName, FStringView InDebugName)
+	: IVirtualizationBackend(ConfigName, InDebugName, EOperations::Both)
+	, Namespace(TEXT("mirage"))
+	, Bucket(TEXT("default"))
+	, ChunkSize(-1)
+	, FailedLoginAttempts(0)		
+{		
+}
+
+bool FHttpBackend::Initialize(const FString& ConfigEntry)
+{
+	using namespace Utility;
+
+	// Some fields are required and will give fatal errors if not found!
+
+	if (!FParse::Value(*ConfigEntry, TEXT("Host="), HostAddress))
+	{
+		UE_LOG(LogVirtualization, Fatal, TEXT("'Host=' not found in the config file"));
+		return false;
+	}
+
+	if (!FParse::Value(*ConfigEntry, TEXT("Namespace="), Namespace))
+	{
+		UE_LOG(LogVirtualization, Fatal, TEXT("'Namespace=' not found in the config file"));
+		return false;
+	}
+
+	if (FParse::Value(*ConfigEntry, TEXT("ChunkSize="), ChunkSize))
+	{
+		UE_LOG(LogVirtualization, Log, TEXT("ChunkSize set to '%" UINT64_FMT "' bytes"), ChunkSize);
+	}
+	else
+	{
+		UE_LOG(LogVirtualization, Log, TEXT("Payloads will not be chunked!"));
+	}
+
+	// If we are connecting to a locally hosted HordeStorage then we do not need authorization
+	if (!IsUsingLocalHost())
+	{
+		if (!FParse::Value(*ConfigEntry, TEXT("OAuthProvider="), OAuthProvider))
+		{
+			UE_LOG(LogVirtualization, Fatal, TEXT("'OAuthProvider=' not found in the config file"));
+			return false;
+		}
+
+		if (!FParse::Value(*ConfigEntry, TEXT("OAuthSecret="), OAuthSecret))
+		{
+			UE_LOG(LogVirtualization, Fatal, TEXT("'OAuthSecret=' not found in the config file"));
+			return false;
+		}
+
+		if (!FParse::Value(*ConfigEntry, TEXT("OAuthClientId="), OAuthClientId))
+		{
+			UE_LOG(LogVirtualization, Fatal, TEXT("'OAuthClientId=' not found in the config file"));
+			return false;
+		}
+	}
+
+	UE_LOG(LogVirtualization, Log, TEXT("Attempting to connect to HordeStorage at '%s' with namespace '%s'"), *HostAddress, *Namespace);
+
+	if (!IsServiceReady())
+	{
+		return false;
+	}
+
+	if (!AcquireAccessToken())
+	{
+		return false;
+	}
+
+	if (!ValidateServiceVersion())
+	{
+		return false;
+	}
+
+	RequestPool = MakeUnique<FRequestPool>(*HostAddress, AccessToken.Get());
+
+	return true;
+}
+
+EPushResult FHttpBackend::PushData(const FPayloadId& Id, const FCompressedBuffer& CompressedPayload, const FString& PackageContext)
+{
+	using namespace Utility;
+
+	TRACE_CPUPROFILER_EVENT_SCOPE(FHttpBackend::PushData);
+
+#if UE_CHECK_FOR_EXISTING_PAYLOADS
+	if (DoesPayloadExist(Id))
+	{
+		UE_LOG(LogVirtualization, Verbose, TEXT("HordeStorage already has a copy of the payload '%s'"), *Id.ToString());
+		return EPushResult::PayloadAlreadyExisted;
+	}
+#endif // UE_CHECK_FOR_EXISTING_PAYLOADS
+
+	FEuropaDDCCachePUTRequest PUTRequest;
+
+	// TODO: This is a waste, we shouldn't need to flatten the buffer, but it makes it easier to work
+	// with the existing chunked code. The chunking code is most likely going to be removed before 
+	// this backend goes to production so it is not worth fixing up to work with FCompressedBuffer 
+	// properly.
+	FSharedBuffer FlattenedPayload = CompressedPayload.GetCompressed().ToShared();
+
+	const int64 NumChunks = FMath::DivideAndRoundUp(FlattenedPayload.GetSize(), ChunkSize);
+	if (NumChunks > MAX_int32)
+	{
+		UE_LOG(LogVirtualization, Error, TEXT("Too many chunks (%d) are required for the payload '%s', try increasing the ChunkSize"), NumChunks, *Id.ToString());
+		return EPushResult::Failed;
+	}
+
+	PUTRequest.ChunkHashes.SetNum((int32)NumChunks);
+
+	const uint8* DataPtr = (const uint8*)FlattenedPayload.GetData();
+
+	std::atomic<int32> NumFailedChunks(0);
+	FGraphEventArray Tasks;
+	Tasks.Reserve((int32)NumChunks);
+
+	// Create and process the chunks that make up the payload
+	for (int32 Index = 0; Index < NumChunks; ++Index)
+	{
+		const int64 ChunkStart = Index * ChunkSize;
+		const uint64 BytesInChunk = FMath::Min(ChunkSize, FlattenedPayload.GetSize() - ChunkStart);
+
+		TArrayView<const uint8> ChunkData(&DataPtr[ChunkStart], (int32)BytesInChunk);
+		FString& ChunkHashString = PUTRequest.ChunkHashes[Index];
+
+		auto Job = [this, ChunkData, &Id, &ChunkHashString, &NumFailedChunks]()
+		{
+			if (!PostChunk(ChunkData, Id, ChunkHashString))
+			{
+				NumFailedChunks++;
+			}
+		};
+
+#if UE_ENABLE_ASYNC_CHUNK_ACCESS
+		Tasks.Add(FFunctionGraphTask::CreateAndDispatchWhenReady(MoveTemp(Job)));
+#else
+		Job();
+#endif //UE_ENABLE_ASYNC_CHUNK_ACCESS
+	}
+
+	// There is some expensive work that we can do while we wait for the chunks to finish their upload
+	const FIoHash PayloadHash = FIoHash::HashBuffer(FlattenedPayload.GetData(), FlattenedPayload.GetSize());
+
+	PUTRequest.PayloadHash = LexToString(PayloadHash);
+	PUTRequest.MetaData.PayloadLength = (int64)FlattenedPayload.GetSize();
+	PUTRequest.MetaData.ChunkLength = ChunkSize;
+
+	{
+		TRACE_CPUPROFILER_EVENT_SCOPE(FHttpBackend::PullData::WaitOnChunks);
+		FTaskGraphInterface::Get().WaitUntilTasksComplete(Tasks);
+	}
+
+	if (NumFailedChunks > 0)
+	{
+		UE_LOG(LogVirtualization, Error, TEXT("Failed to upload %d chunks for the payload '%s'."), NumFailedChunks.load(), *Id.ToString());
+		return EPushResult::Failed;
+	}
+
+	UE_LOG(LogVirtualization, Verbose, TEXT("Successfully uploaded all chunks for the payload '%s'"), *Id.ToString());
+
+	// Note that the ddc end point is used by both ddc and mirage
+	TStringBuilder<256> Uri;
+	Uri.Appendf(TEXT("api/v1/c/ddc/%s/%s/%s"), *Namespace, *Bucket, *Id.ToString());
+
+	// Retry request until we get an accepted response or exhaust allowed number of attempts.
+	uint32 Attempts = 0;
+	while (++Attempts <= UE_MIRAGE_MAX_ATTEMPTS)
+	{
+		FScopedRequestPtr Request(RequestPool.Get());
+		if (Request.IsValid())
+		{
+			const FString OutputString = PUTRequest.ToJson(false);
+
+			TStringConversion<TStringConvert<TCHAR, ANSICHAR>> ConvertedString = StringCast<ANSICHAR>(*OutputString);
+			TArrayView<const uint8> AnsiStringBuffer((uint8*)ConvertedString.Get(), ConvertedString.Length());
+
+			FRequest::Result Result = Request->PerformBlockingUpload<FRequest::PutJson>(*Uri, AnsiStringBuffer);
+			const int64 ResponseCode = Request->GetResponseCode();
+
+			if (ResponseCode == 200)
+			{
+				UE_LOG(LogVirtualization, Verbose, TEXT("Successfully uploaded the description for the payload '%s'"), *Id.ToString());
+				return EPushResult::Success;
+			}
+
+			if (!ShouldRetryOnError(ResponseCode))
+			{
+				UE_LOG(LogVirtualization, Error, TEXT("Failed with error code '%d' to upload header infomation about payload '%s'"), ResponseCode, *Id.ToString());
+				return EPushResult::Failed;
+			}
+		}
+	}
+
+	UE_LOG(LogVirtualization, Error, TEXT("Failed  '%d' attempts to upload header infomation about payload '%s'"), UE_MIRAGE_MAX_ATTEMPTS, *Id.ToString());
+	return EPushResult::Failed;
+}
+
+FCompressedBuffer FHttpBackend::PullData(const FPayloadId& Id)
+{
+	using namespace Utility;
+
+	TRACE_CPUPROFILER_EVENT_SCOPE(FHttpBackend::PullData);
+
+	// First we need to get the description of the payload from Europa
+
+	// Note that the ddc end point is used by both ddc and mirage
+	// fields=contentSha1 - Ask for the Sha1 hash of the fully reconstructed payload
+	// fields=blobIdentifiers - As for a list of the Sha1 hash ids for the payload chunks that we need to access from
+	// the Io service, the ids will be in the correct order.
+	// fields=metadata - Ask for the payload metadata which contains info we can use later for optimizations.
+	TStringBuilder<256> Uri;
+	Uri.Appendf(TEXT("api/v1/c/ddc/%s/%s/%s.json?fields=contentHash&fields=blobIdentifiers&fields=metadata"), *Namespace, *Bucket, *Id.ToString());
+
+	FDDCCacheGETResponse Response;
+
+	// Retry request until we get an accepted response or exhaust allowed number of attempts.
+	int64 ResponseCode = 0;
+	uint32 Attempts = 0;
+	while (ResponseCode != 200 && ++Attempts <= UE_MIRAGE_MAX_ATTEMPTS)
+	{
+		FScopedRequestPtr Request(RequestPool.Get());
+		if (Request.IsValid())
+		{
+			const FRequest::Result Result = Request->PerformBlockingDownload(*Uri);
+			ResponseCode = Request->GetResponseCode();
+
+			// Request was successful, make sure we got all the expected data.
+			if (FRequest::IsSuccessfulResponse(ResponseCode))
+			{
+				if (!Response.FromJson(Request->GetResponseAsJsonObject()))
+				{
+					UE_LOG(LogVirtualization, Error, TEXT("Failed to parser the header infomation about payload '%s'"), *Id.ToString());
+					return FCompressedBuffer();
+				}
+			}
+			else if (ResponseCode == 400)
+			{
+				// Response 400 indicates that the payload does not exist in HordeStorage. Note that it is faster to just make the request
+				// and check for the response rather than call ::DoesPayloadExist prior to requesting the json header because this way 
+				// we will only make a single request if the payload exists or not.
+				UE_LOG(LogVirtualization, Verbose, TEXT("[%s] Does not contain the payload '%s'"), *GetDebugName(), *Id.ToString());
+				return FCompressedBuffer();
+			}
+			else if (!ShouldRetryOnError(ResponseCode))
+			{
+				UE_LOG(LogVirtualization, Error, TEXT("Failed with error code '%d' to download header infomation about payload '%s'"), ResponseCode, *Id.ToString());
+				return FCompressedBuffer();
+			}
+		}
+	}
+
+	if (ResponseCode != 200)
+	{
+		UE_LOG(LogVirtualization, Error, TEXT("Failed '%d' attempts to download header infomation about payload (last error code '%d')  '%s'"), UE_MIRAGE_MAX_ATTEMPTS, ResponseCode, *Id.ToString());
+		return FCompressedBuffer();
+	}
+
+	UE_LOG(LogVirtualization, Verbose, TEXT("Successfully downloaded a description for the payload '%s'"), *Id.ToString());
+
+	// Now that we have the payload description we can start pulling the chunks from the Io service
+	// and reconstruct the final payload.
+	FUniqueBuffer Payload = FUniqueBuffer::Alloc(Response.MetaData.PayloadLength);
+
+	int64 BytesLeft = Response.MetaData.PayloadLength;
+
+	FGraphEventArray Tasks;
+	Tasks.Reserve(Response.ChunkHashes.Num());
+
+	std::atomic<int32> NumFailedChunks(0);
+
+	bool bAllChunksPulled = true;
+	uint8* PayloadPtr = (uint8*)Payload.GetData();
+
+	for (const FString& HashString : Response.ChunkHashes)
+	{
+		checkf(BytesLeft > 0, TEXT("Ran out of buffer space before all payload chunks were read!"));
+
+		const int64 BytesToRead = FMath::Min(BytesLeft, Response.MetaData.ChunkLength);
+
+		auto Job = [this, PayloadPtr, BytesToRead, &Id, &HashString, &NumFailedChunks]()
+		{
+			if (!PullChunk(HashString, Id, PayloadPtr, BytesToRead))
+			{
+				NumFailedChunks++;
+			}
+		};
+
+#if UE_ENABLE_ASYNC_CHUNK_ACCESS
+		Tasks.Add(FFunctionGraphTask::CreateAndDispatchWhenReady(MoveTemp(Job)));
+#else
+		Job();
+#endif //UE_ENABLE_ASYNC_CHUNK_ACCESS
+
+		PayloadPtr += BytesToRead;
+		BytesLeft -= BytesToRead;
+	}
+
+	{
+		TRACE_CPUPROFILER_EVENT_SCOPE(FHttpBackend::PullData::WaitOnChunks);
+		FTaskGraphInterface::Get().WaitUntilTasksComplete(Tasks);
+	}
+
+	if (NumFailedChunks == 0)
+	{
+		UE_LOG(LogVirtualization, Verbose, TEXT("Successfully downloaded all chunks for the payload '%s'"), *Id.ToString());	
+		return FCompressedBuffer::FromCompressed(Payload.MoveToShared());
+	}
+	else
+	{
+		UE_LOG(LogVirtualization, Error, TEXT("Failed to download %d chunks for the payload '%s'"), NumFailedChunks.load(), *Id.ToString());
+		return FCompressedBuffer();
+	}
+}
+
+bool FHttpBackend::DoesPayloadExist(const FPayloadId& Id)
+{
+	using namespace Utility;
+
+	TRACE_CPUPROFILER_EVENT_SCOPE(FHttpBackend::DoesPayloadExist);
+
+	// Note that the ddc end point is used by both ddc and mirage
+	TStringBuilder<256> Uri;
+	Uri.Appendf(TEXT("api/v1/c/ddc/%s/%s/%s"), *Namespace, *Bucket, *Id.ToString());
+
+	// Retry request until we get an accepted response or exhaust allowed number of attempts.
+	uint32 Attempts = 0;
+	while (++Attempts <= UE_MIRAGE_MAX_ATTEMPTS)
+	{
+		FScopedRequestPtr Request(RequestPool.Get());
+		if (Request.IsValid())
+		{
+			const FRequest::Result Result = Request->PerformBlockingQuery<FRequest::Head>(*Uri);
+			const int64 ResponseCode = Request->GetResponseCode();
+
+			if (FRequest::IsSuccessfulResponse(ResponseCode))
+			{
+				return true;
+			}
+			else if (ResponseCode == 400)
+			{
+				return false;
+			}
+
+			if (!ShouldRetryOnError(ResponseCode))
+			{
+				return false;
+			}
+		}
+	}
+
+	return false;
+}
+
+bool FHttpBackend::IsUsingLocalHost() const
+{
+	return HostAddress.StartsWith(TEXT("http://localhost"));
+}
+
+bool FHttpBackend::IsServiceReady() const
+{
+	// TODO: Pretty much the same code as in FHttpDerivedDataBackend, another candidate for code sharing
+	TRACE_CPUPROFILER_EVENT_SCOPE(FHttpBackend::IsServiceReady);
+
+	using namespace Utility;
+
+	FRequest Request(*HostAddress, nullptr, false);
+	FRequest::Result Result = Request.PerformBlockingDownload(TEXT("health/ready"), nullptr);
+
+	if (Result == FRequest::Success && FRequest::IsSuccessfulResponse(Request.GetResponseCode()))
+	{
+		UE_LOG(LogVirtualization, Log, TEXT("HordeStorage status: '%s'."), *Request.GetResponseAsString());
+		return true;
+	}
+	else
+	{
+		UE_LOG(LogVirtualization, Error, TEXT("Unable to reach HordeStorage at '%s'. Status: %d . Response: '%s'"), *HostAddress, Request.GetResponseCode(), *Request.GetResponseAsString());
+		return false;
+	}
+}
+
+bool FHttpBackend::AcquireAccessToken()
+{
+	// TODO: Pretty much the same code as in FHttpDerivedDataBackend, another candidate for code sharing
+	TRACE_CPUPROFILER_EVENT_SCOPE(FHttpBackend::AcquireAccessToken);
+
+	using namespace Utility;
+
+	if (IsUsingLocalHost())
+	{
+		UE_LOG(LogVirtualization, Log, TEXT("Connecting to a local host '%s', so skipping authorization"), *HostAddress);
+		return true;
+	}
+
+	// Avoid spamming the this if the service is down
+	if (FailedLoginAttempts > UE_MIRAGE_MAX_FAILED_LOGIN_ATTEMPTS)
+	{
+		return false;
+	}
+
+	ensureMsgf(OAuthProvider.StartsWith(TEXT("http://")) || OAuthProvider.StartsWith(TEXT("https://")),
+		TEXT("The OAuth provider %s is not valid. Needs to be a fully qualified url."),
+		*OAuthProvider
+	);
+
+	// In case many requests wants to update the token at the same time
+	// get the current serial while we wait to take the CS.
+	const uint32 WantsToUpdateTokenSerial = AccessToken.IsValid() ? AccessToken->GetSerial() : 0u;
+
+	{
+		FScopeLock Lock(&AccessCs);
+
+		// Check if someone has beaten us to update the token, then it 
+		// should now be valid.
+		if (AccessToken.IsValid() && AccessToken->GetSerial() > WantsToUpdateTokenSerial)
+		{
+			return true;
+		}
+
+		const uint32 SchemeEnd = OAuthProvider.Find(TEXT("://")) + 3;
+		const uint32 DomainEnd = OAuthProvider.Find(TEXT("/"), ESearchCase::CaseSensitive, ESearchDir::FromStart, SchemeEnd);
+		FString AuthDomain(DomainEnd, *OAuthProvider);
+		FString Uri(*OAuthProvider + DomainEnd + 1);
+
+		FRequest Request(*AuthDomain, nullptr, false);
+
+		// If contents of the secret string is a file path, resolve and read form data.
+		if (OAuthSecret.StartsWith(TEXT("\\\\")))
+		{
+			FString SecretFileContents;
+			if (FFileHelper::LoadFileToString(SecretFileContents, *OAuthSecret))
+			{
+				// Overwrite the filepath with the actual content.
+				OAuthSecret = SecretFileContents;
+			}
+			else
+			{
+				UE_LOG(LogVirtualization, Warning, TEXT("Failed to read OAuth form data file (%s)."), *OAuthSecret);
+				return false;
+			}
+		}
+
+		FString OAuthFormData = FString::Printf(
+			TEXT("client_id=%s&scope=cache_access&grant_type=client_credentials&client_secret=%s"),
+			*OAuthClientId,
+			*OAuthSecret
+		);
+
+		TArray<uint8> FormData;
+		auto OAuthFormDataUTF8 = FTCHARToUTF8(*OAuthFormData);
+		FormData.Append((uint8*)OAuthFormDataUTF8.Get(), OAuthFormDataUTF8.Length());
+
+		FRequest::Result Result = Request.PerformBlockingUpload<FRequest::Post>(*Uri, MakeArrayView(FormData));
+
+		if (Result == FRequest::Success && Request.GetResponseCode() == 200)
+		{
+			TSharedPtr<FJsonObject> ResponseObject = Request.GetResponseAsJsonObject();
+			if (ResponseObject)
+			{
+				FString AccessTokenString;
+				int32 ExpiryTimeSeconds = 0;
+				int32 CurrentTimeSeconds = int32(FPlatformTime::ToSeconds(FPlatformTime::Cycles()));
+
+				if (ResponseObject->TryGetStringField(TEXT("access_token"), AccessTokenString) &&
+					ResponseObject->TryGetNumberField(TEXT("expires_in"), ExpiryTimeSeconds))
+				{
+					if (!AccessToken)
+					{
+						AccessToken = MakeUnique<FAccessToken>();
+					}
+					AccessToken->SetHeader(*AccessTokenString);
+					UE_LOG(LogVirtualization, Log, TEXT("Logged in to HTTP DDC services. Expires in %d seconds."), ExpiryTimeSeconds);
+
+					//Schedule a refresh of the token ahead of expiry time (this will not work in commandlets)
+					if (!IsRunningCommandlet())
+					{
+						FTSTicker::GetCoreTicker().AddTicker(FTickerDelegate::CreateLambda(
+							[this](float DeltaTime)
+							{
+								this->AcquireAccessToken();
+								return false;
+							}
+						), ExpiryTimeSeconds - 20.0f);
+					}
+					// Reset failed login attempts, the service is indeed alive.
+					FailedLoginAttempts = 0;
+					return true;
+				}
+			}
+		}
+		else
+		{
+			UE_LOG(LogVirtualization, Warning, TEXT("Failed to log in to HTTP services. Server responed with code %d."), Request.GetResponseCode());
+			FailedLoginAttempts++;
+		}
+	}
+	return false;
+}
+
+bool FHttpBackend::ValidateServiceVersion()
+{
+	TRACE_CPUPROFILER_EVENT_SCOPE(FHttpBackend::ValidateServiceVersion);
+
+	using namespace Utility;
+
+	FHttpServiceStatus Response;
+
+	int64 ResponseCode = 0;
+	uint32 Attempts = 0;
+	while (ResponseCode != 200 && ++Attempts <= UE_MIRAGE_MAX_ATTEMPTS)
+	{
+		// We create the FRequest ourselves since the request pool does not yet exist
+		FRequest Request(*HostAddress, AccessToken.Get(), true);
+		FRequest::Result Result = Request.PerformBlockingDownload(TEXT("api/v1/status"));
+
+		ResponseCode = Request.GetResponseCode();
+
+		// Request was successful, make sure we got all the expected data.
+		if (ResponseCode == 200)
+		{
+			if (!Response.FromJson(Request.GetResponseAsJsonObject()))
+			{
+				UE_LOG(LogVirtualization, Error, TEXT("The response to 'api/v1/status' GET did not contain valid data!"));
+				return false;
+			}
+		}
+		else if (!ShouldRetryOnError(ResponseCode))
+		{
+			UE_LOG(LogVirtualization, Error, TEXT("Failed with error code '%d' to access the services status"), ResponseCode);
+			return false;
+		}
+	}
+
+	if (ResponseCode != 200)
+	{
+		UE_LOG(LogVirtualization, Error, TEXT("Failed '%d' attempts to access the services status (last error code '%d')"), UE_MIRAGE_MAX_ATTEMPTS, ResponseCode);
+		return false;
+	}
+
+	// Check version number
+	if (!Response.DoesHaveValidVersion(UE_HORDESTORAGE_MIN_MAJOR_VER, UE_HORDESTORAGE_MIN_MINOR_VER, UE_HORDESTORAGE_MIN_PATCH_VER))
+	{
+		return false;
+	}
+
+	if (!Response.SupportsCapability(TEXT("ddc")))
+	{
+		UE_LOG(LogVirtualization, Error, TEXT("HordeStorage does not support Europa (ddc) capability"));
+		return false;
+	}
+
+	Response.LogStatusInfo();
+
+	return true;
+}
+	
+bool FHttpBackend::ShouldRetryOnError(int64 ResponseCode)
+{
+	// TODO: Pretty much the same code as in FHttpDerivedDataBackend, another candidate for code sharing
+	// 
+	// Access token might have expired, request a new token and try again.
+	if (ResponseCode == 401 && AcquireAccessToken())
+	{
+		return true;
+	}
+
+	// Too many requests, make a new attempt
+	if (ResponseCode == 429)
+	{
+		return true;
+	}
+
+	// Gateway timeout, it will most likely work if we try again
+	if (ResponseCode == 504)
+	{
+		return true;
+	}
+
+	return false;
+}
+
+bool FHttpBackend::PostChunk(const TArrayView<const uint8>& ChunkData, const FPayloadId& PayloadId, FString& OutHashAsString)
+{
+	TRACE_CPUPROFILER_EVENT_SCOPE(FHttpBackend::PostChunk);
+
+	using namespace Utility;
+
+	const FIoHash ChunkHash = FIoHash::HashBuffer(ChunkData.GetData(), ChunkData.Num());
+	OutHashAsString = LexToString(ChunkHash);
+
+#if UE_CHECK_FOR_EXISTING_CHUNKS
+	if (DoesChunkExist(HashAsString))
+	{
+		UE_LOG(LogVirtualization, Verbose, TEXT("HordeStorage already has a copy of the chunk '%s' for payload '%s'"), *HashAsString, *PayloadId.ToString());
+		return true;
+	}
+#endif // UE_CHECK_FOR_EXISTING_CHUNKS
+
+	FScopedRequestPtr Request(RequestPool.Get());
+	// TODO: Another candidate for code sharing
+	Request->SetHeader(TEXT("X-Jupiter-IoHash"), *OutHashAsString);
+
+	TStringBuilder<256> Uri;
+	Uri.Appendf(TEXT("api/v1/s/%s/%s"), *Namespace, *OutHashAsString);
+
+	// Retry request until we get an accepted response or exhaust allowed number of attempts.
+	uint32 Attempts = 0;
+	while (++Attempts <= UE_MIRAGE_MAX_ATTEMPTS)
+	{
+		if (Request.IsValid())
+		{
+			const FRequest::Result Result = Request->PerformBlockingUpload<FRequest::Put>(*Uri, ChunkData);
+			const int64 ResponseCode = Request->GetResponseCode();
+
+			if (ResponseCode == 200)
+			{
+				UE_LOG(LogVirtualization, Verbose, TEXT("Successfully uploaded a chunk '%s'for payload '%s'"), *OutHashAsString, *PayloadId.ToString());
+				return true;
+			}
+
+			if (!ShouldRetryOnError(ResponseCode))
+			{
+				return false;
+			}
+		}
+	}
+
+	return false;
+}
+
+bool FHttpBackend::PullChunk(const FString& Hash, const FPayloadId& PayloadId, uint8* DataPtr, int64 BufferSize)
+{
+	TRACE_CPUPROFILER_EVENT_SCOPE(FHttpBackend::PullChunk);
+
+	using namespace Utility;
+
+	FScopedRequestPtr Request(RequestPool.Get());
+
+	TStringBuilder<256> Uri;
+	Uri.Appendf(TEXT("api/v1/s/%s/%s"), *Namespace, *Hash);
+
+	// Retry request until we get an accepted response or exhaust allowed number of attempts.
+	uint32 Attempts = 0;
+	while (++Attempts <= UE_MIRAGE_MAX_ATTEMPTS)
+	{
+		if (Request.IsValid())
+		{
+			const FRequest::Result Result = Request->PerformBlockingDownload(*Uri, FMutableMemoryView(DataPtr, BufferSize));
+
+			if (Result != FRequest::Success)
+			{
+				UE_LOG(LogVirtualization, Error, TEXT("Attempting to GET a payload chunk '%s' for payload '%s' failed due to an internal Curl error"), *Hash, *PayloadId.ToString());	
+				return false;
+			}
+
+			const int64 ResponseCode = Request->GetResponseCode();
+
+			if (ResponseCode == 200)
+			{
+				UE_LOG(LogVirtualization, Verbose, TEXT("Successfully downloaded a payload chunk '%s' for payload '%s'"), *Hash, *PayloadId.ToString());	
+				return true;
+			}
+
+			if (!ShouldRetryOnError(ResponseCode))
+			{
+				UE_LOG(LogVirtualization, Error, TEXT("Attempting to GET a payload chunk '%s' for payload '%s' failed with http response: %" INT64_FMT), *Hash, *PayloadId.ToString(), ResponseCode);	
+				return false;
+			}
+		}
+	}
+
+	UE_LOG(LogVirtualization, Error, TEXT("Attempting to GET a payload chunk '%s' for payload '%s' failed all '%d' attempts"), *Hash, *PayloadId.ToString(), UE_MIRAGE_MAX_ATTEMPTS);	
+	return false;
+}
+
+bool FHttpBackend::DoesChunkExist(const FString& Hash)
+{
+	using namespace Utility;
+
+	TRACE_CPUPROFILER_EVENT_SCOPE(FHttpBackend::DoesChunkExist);
+
+	TStringBuilder<256> Uri;
+	Uri.Appendf(TEXT("api/v1/s/%s/%s"), *Namespace, *Hash);
+
+	// Retry request until we get an accepted response or exhaust allowed number of attempts.
+	uint32 Attempts = 0;
+	while (++Attempts <= UE_MIRAGE_MAX_ATTEMPTS)
+	{
+		FScopedRequestPtr Request(RequestPool.Get());
+		if (Request.IsValid())
+		{
+			const FRequest::Result Result = Request->PerformBlockingQuery<FRequest::Head>(*Uri);
+			const int64 ResponseCode = Request->GetResponseCode();
+
+			if (FRequest::IsSuccessfulResponse(ResponseCode))
+			{
+				return true;
+			}
+			else if (ResponseCode == 404)
+			{
+				return false;
+			}
+
+			if (!ShouldRetryOnError(ResponseCode))
+			{
+				return false;
+			}
+		}
+	}
+
+	return false;
+}
+
+UE_REGISTER_VIRTUALIZATION_BACKEND_FACTORY(FHttpBackend, HordeStorage);
+
+} // namespace Virtualization
+} // namespace UE
+
+#endif //UE_USE_HORDESTORAGE_BACKEND
diff --git a/Engine/Source/Developer/Virtualization/Private/VirtualizationHttpBackend.h b/Engine/Source/Developer/Virtualization/Private/VirtualizationHttpBackend.h
new file mode 100644
--- /dev/null
+++ b/Engine/Source/Developer/Virtualization/Private/VirtualizationHttpBackend.h
@@ -0,0 +1,91 @@
+// Copyright Epic Games, Inc. All Rights Reserved.
+
+#pragma once
+
+#include "IVirtualizationBackend.h"
+
+namespace UE::Utility { struct FRequestPool; }
+namespace UE::Utility { struct FAccessToken; }
+
+namespace UE::Virtualization
+{
+/**
+* This backend allows data to be stored in and retrieved from the Horde storage service.
+*
+* Ini file setup:
+* 'Name'=(Type=HordeStorage, Host="", Namespace="", ChunkSize=, OAuthProvider="", OAuthClientId="", OAuthSecret="")
+* Host:				The URL of the service, use http://localhost if hosted locally.
+* Namespace:		Horde storage is divided into a number of namespaces allowing projects to keep their data separate
+*					while using the same service. This value controls which name space will be used.
+* ChunkSize:		Each payload can be divided into a number of chunks when being uploaded to Horde to improve upload
+*					performance, this value sets the max size (in bytes) of each chunk. To disable and attempt to upload
+*					each payload as a single data blob, set this to -1.
+* OAuthProvider:	Url of the OAuth authorization server.
+* OAuthClientId:	Public identifier for use with the OAuth authorization server.
+* OAuthSecret:		Password for the OAuthClientId
+* (Note that the OAuth entries are not required if hosting locally)
+*/
+class FHttpBackend final : public IVirtualizationBackend
+{
+public:
+	explicit FHttpBackend(FStringView ConfigName, FStringView DebugName);
+	virtual ~FHttpBackend() = default;
+
+private:
+	/* IVirtualizationBackend implementation */
+
+	virtual bool Initialize(const FString& ConfigEntry) override;
+
+	virtual EPushResult PushData(const FPayloadId& Id, const FCompressedBuffer& CompressedPayload, const FString& PackageContext) override;
+
+	virtual FCompressedBuffer PullData(const FPayloadId& Id) override;
+
+	virtual bool DoesPayloadExist(const FPayloadId& Id) override;
+
+private:
+
+	bool IsUsingLocalHost() const;
+	bool IsServiceReady() const;
+	bool AcquireAccessToken();
+
+	/**
+	 * Request the status of the service that we are connected to and make sure that it supports the
+	 * feature set we need and meets our minimum version requirements.
+	 */
+	bool ValidateServiceVersion();
+
+	bool ShouldRetryOnError(int64 ResponseCode);
+
+	bool PostChunk(const TArrayView<const uint8>& ChunkData, const FPayloadId& PayloadId, FString& OutHashAsString);
+	bool PullChunk(const FString& Hash, const FPayloadId& PayloadId, uint8* DataPtr, int64 BufferSize);
+	bool DoesChunkExist(const FString& Hash);
+
+	/** Address of the service*/
+	FString HostAddress;
+	/** Namespace to connect to */
+	FString Namespace;
+	/** Europa allows us to organize the payloads by bucket. Currently this is not exposed and just set to 'default' */
+	FString Bucket;
+
+	/** The max size of each payload chunk */
+	uint64 ChunkSize;
+
+	/** Url of the OAuth authorization server */
+	FString OAuthProvider;
+	/**  Public identifier for use with the OAuth authorization server */
+	FString OAuthClientId;
+	/** Password for the OAuthClientId */
+	FString OAuthSecret;
+
+	/** The pool of FRequest objects that can be recycled */
+	TUniquePtr<Utility::FRequestPool> RequestPool;
+
+	/** Critical section used to protect the creation of new access tokens */
+	FCriticalSection AccessCs;
+	/** The access token used with service authorization */
+	TUniquePtr<Utility::FAccessToken> AccessToken;
+	/** Count how many times a login has failed since the last successful login */
+	uint32 FailedLoginAttempts;
+};
+
+} // namespace UE::Virtualization
diff --git a/Engine/Source/Developer/Virtualization/Private/VirtualizationJupiterBackend.cpp b/Engine/Source/Developer/Virtualization/Private/VirtualizationJupiterBackend.cpp
deleted file mode 100644
diff --git a/Engine/Source/Developer/Virtualization/Private/VirtualizationJupiterBackend.h b/Engine/Source/Developer/Virtualization/Private/VirtualizationJupiterBackend.h
deleted file mode 100644
