Give improved error messages to the user if the source control asset virtualization backend fails to connect.

#rnx

- If a perforce connection issue is encountered when starting up the editor we will now print an error to the message log and then on the first tick of the editor display a notification to get the users attention. It will no longer result in a fatal error.

- Add a new category "Asset Virtualization" to the message log so that we can start to improve the feedback we give to end users rather than only reporting problems via UE_LOG. This also introduces a new localization text name space.
-- I will need to go over a lot of the virtualization modules error handling and move it to use the message log instead.
- When the source control backend first starts up we now specifically check if the connection works and if not try to give the user feedback on how to fix it.
-- Ideally we'd provide a  login dialog for the user to fix the problem in place but due to how early in the start up process we need to be able to access the virtualization backends this is quite tricky.
- In the future we will build the concept of a backends connection into a formalized concept in the system along with better error handling. At this point FSourceControlBackend::OnConnectionError should be moved to general use for all backends. this will allow us to provide much better status info to the user.

diff --git a/Engine/Source/Developer/Virtualization/Private/Virtualization.cpp b/Engine/Source/Developer/Virtualization/Private/Virtualization.cpp
--- a/Engine/Source/Developer/Virtualization/Private/Virtualization.cpp
+++ b/Engine/Source/Developer/Virtualization/Private/Virtualization.cpp
@@ -4,11 +4,14 @@
 
 #include "Features/IModularFeatures.h"
 #include "ISourceControlModule.h"
+#include "MessageLogModule.h"
 #include "Misc/DelayedAutoRegister.h"
 #include "PackageSubmissionChecks.h"
 #include "Serialization/VirtualizedBulkData.h"
 #include "VirtualizationSourceControlUtilities.h"
 
+#define LOCTEXT_NAMESPACE "Virtualization"
+
 namespace UE::Virtualization
 {
 
@@ -25,12 +28,19 @@ public:
 				PackageSubmissionHandle = ISourceControlModule::Get().RegisterPreSubmitFinalize(
 					FSourceControlPreSubmitFinalizeDelegate::FDelegate::CreateStatic(&OnPrePackageSubmission));
 			});
+
+		FMessageLogModule& MessageLogModule = FModuleManager::LoadModuleChecked<FMessageLogModule>("MessageLog");
+		MessageLogModule.RegisterLogListing("LogVirtualization", LOCTEXT("AssetVirtualizationLogLabel", "Asset Virtualization"));
 	}
 
 	virtual void ShutdownModule() override
 	{
-		IModularFeatures::Get().UnregisterModularFeature(FName("VirtualizationSourceControlUtilities"), &SourceControlutility);
-		
+		if (FModuleManager::Get().IsModuleLoaded("MessageLog"))
+		{
+			FMessageLogModule& MessageLogModule = FModuleManager::GetModuleChecked<FMessageLogModule>("MessageLog");
+			MessageLogModule.UnregisterLogListing("LogVirtualization&");
+		}
+	
 		// The SourceControl module might be destroyed before this one, depending on shutdown order so we need to check if 
 		// it is loaded before unregistering.
 		if (ISourceControlModule* SourceControlModule = FModuleManager::GetModulePtr<ISourceControlModule>(FName("SourceControl")))
@@ -40,6 +50,7 @@ public:
 
 		PackageSubmissionHandle.Reset();
 		
+		IModularFeatures::Get().UnregisterModularFeature(FName("VirtualizationSourceControlUtilities"), &SourceControlutility);
 	}
 
 private:
@@ -50,4 +61,6 @@ private:
 
 } // namespace UE::Virtualization
 
+#undef LOCTEXT_NAMESPACE
+
 IMPLEMENT_MODULE(UE::Virtualization::FVirtualizationModule, Virtualization);
diff --git a/Engine/Source/Developer/Virtualization/Private/VirtualizationSourceControlBackend.cpp b/Engine/Source/Developer/Virtualization/Private/VirtualizationSourceControlBackend.cpp
--- a/Engine/Source/Developer/Virtualization/Private/VirtualizationSourceControlBackend.cpp
+++ b/Engine/Source/Developer/Virtualization/Private/VirtualizationSourceControlBackend.cpp
@@ -2,9 +2,11 @@
 
 #include "VirtualizationSourceControlBackend.h"
 
+#include "Containers/Ticker.h"
 #include "HAL/FileManager.h"
 #include "ISourceControlModule.h"
 #include "ISourceControlProvider.h"
+#include "Logging/MessageLog.h"
 #include "Misc/App.h"
 #include "Misc/Parse.h"
 #include "Misc/Paths.h"
@@ -18,6 +20,8 @@
 // can enable this and stop using the hacky work around 'TryToDownloadFileFromBackgroundThread'
 #define IS_SOURCE_CONTROL_THREAD_SAFE 0
 
+#define LOCTEXT_NAMESPACE "Virtualization"
+
 namespace UE::Virtualization
 {
 
@@ -88,6 +92,26 @@ bool FSourceControlBackend::Initialize(const FString& ConfigEntry)
 		SCCProvider.Init();
 	}
 
+	// Note that if the connect is failing then we expect it to fail here rather than in the subsequent attempts to get the meta info file
+	TSharedRef<FConnect, ESPMode::ThreadSafe> ConnectCommand = ISourceControlOperation::Create<FConnect>();
+	if (SCCProvider.Execute(ConnectCommand, FString(), EConcurrency::Synchronous) != ECommandResult::Succeeded)
+	{		
+		FTextBuilder Errors;
+		for (const FText& Msg : ConnectCommand->GetResultInfo().ErrorMessages)
+		{
+			Errors.AppendLine(Msg);
+		}
+
+		FMessageLog Log("LogVirtualization");
+		Log.Error(	FText::Format(LOCTEXT("FailedSourceControlConnection", "Failed to connect to source control backend with the following errors:\n{0}\n"
+					"The source control backend will be unable to pull payloads!\n"
+					"Trying logging in with the 'p4 login' command or by using p4vs/UnrealGameSync."),
+					Errors.ToText()));
+		
+		OnConnectionError();
+		return true;
+	}
+
 	// When a source control depot is set up a file named 'payload_metainfo.txt' should be submitted to it's root.
 	// This allows us to check for the existence of the file to confirm that the depot root is indeed valid.
 	const FString PayloadMetaInfoPath = FString::Printf(TEXT("%spayload_metainfo.txt"), *DepotRoot);
@@ -96,23 +120,41 @@ bool FSourceControlBackend::Initialize(const FString& ConfigEntry)
 	TSharedRef<FDownloadFile, ESPMode::ThreadSafe> DownloadCommand = ISourceControlOperation::Create<FDownloadFile>();
 	if (SCCProvider.Execute(DownloadCommand, PayloadMetaInfoPath, EConcurrency::Synchronous) != ECommandResult::Succeeded)
 	{
-		UE_LOG(LogVirtualization, Error, TEXT("Failed to find 'payload_metainfo.txt' in the depot '%s', is your config set up correctly?"), *DepotRoot);
-		return false;
+		FMessageLog Log("LogVirtualization");
+
+		Log.Error(	FText::Format(LOCTEXT("FailedMetaInfo", "Failed to find 'payload_metainfo.txt' in the depot '{0}'\n"
+					"The source control backend will be unable to pull payloads, is your source control  config set up correctly?"),
+					FText::FromString(DepotRoot)));
+		
+		OnConnectionError();
+		return true;
 	}	
 #else
 	TSharedRef<FDownloadFile, ESPMode::ThreadSafe> DownloadCommand = ISourceControlOperation::Create<FDownloadFile>();
 	if (!SCCProvider.TryToDownloadFileFromBackgroundThread(DownloadCommand, PayloadMetaInfoPath))
 	{
-		UE_LOG(LogVirtualization, Error, TEXT("Failed to find 'payload_metainfo.txt' in the depot '%s', is your config set up correctly?"), *DepotRoot);
-		return false;
+		FMessageLog Log("LogVirtualization");
+
+		Log.Error(	FText::Format(LOCTEXT("FailedMetaInfo", "Failed to find 'payload_metainfo.txt' in the depot '{0}'\n"
+					"The source control backend will be unable to pull payloads, is your source control  config set up correctly?"),
+					FText::FromString(DepotRoot)));
+		
+		OnConnectionError();
+		return true;
 	}		
 #endif //IS_SOURCE_CONTROL_THREAD_SAFE
 
 	FSharedBuffer MetaInfoBuffer = DownloadCommand->GetFileData(PayloadMetaInfoPath);
 	if (MetaInfoBuffer.IsNull())
 	{
-		UE_LOG(LogVirtualization, Error, TEXT("Failed to find 'payload_metainfo.txt' in the depot '%s', is your config set up correctly?"), *DepotRoot);
-		return false;
+		FMessageLog Log("LogVirtualization");
+
+		Log.Error(	FText::Format(LOCTEXT("FailedMetaInfo", "Failed to find 'payload_metainfo.txt' in the depot '{0}'\n"
+					"The source control backend will be unable to pull payloads, is your source control  config set up correctly?"),
+					FText::FromString(DepotRoot)));
+
+		OnConnectionError();
+		return true;
 	}
 
 	// Currently we do not do anything with the payload meta info, in the future we could structure
@@ -453,6 +495,22 @@ void FSourceControlBackend::CreateDepotPath(const FPayloadId& PayloadId, FString
 	OutPath << DepotRoot << PayloadPath;
 }
 
+void FSourceControlBackend::OnConnectionError()
+{
+	auto Callback = [](float Delta)->bool
+	{
+		FMessageLog Log("LogVirtualization");
+		Log.Notify(LOCTEXT("ConnectionError", "Asset virtualization connect errors were encountered, see the message log for more info"));
+
+		// This tick callback is one shot, so return false to prevent it being invoked again
+		return false;
+	};
+
+	FTSTicker::GetCoreTicker().AddTicker(FTickerDelegate::CreateLambda(Callback));
+}
+
 UE_REGISTER_VIRTUALIZATION_BACKEND_FACTORY(FSourceControlBackend, SourceControl);
 
 } // namespace UE::Virtualization
+
+#undef LOCTEXT_NAMESPACE
diff --git a/Engine/Source/Developer/Virtualization/Private/VirtualizationSourceControlBackend.h b/Engine/Source/Developer/Virtualization/Private/VirtualizationSourceControlBackend.h
--- a/Engine/Source/Developer/Virtualization/Private/VirtualizationSourceControlBackend.h
+++ b/Engine/Source/Developer/Virtualization/Private/VirtualizationSourceControlBackend.h
@@ -51,6 +51,9 @@ private:
 
 	void CreateDepotPath(const FPayloadId& PayloadId, FStringBuilderBase& OutPath);
 
+	/** Will display a FMessage notification to the user on the next valid engine tick to try and keep them aware of connection failures */
+	void OnConnectionError();
+	
 	/** The root where the virtualized payloads are stored in source control */
 	FString DepotRoot;
 
diff --git a/Engine/Source/Developer/Virtualization/Virtualization.Build.cs b/Engine/Source/Developer/Virtualization/Virtualization.Build.cs
--- a/Engine/Source/Developer/Virtualization/Virtualization.Build.cs
+++ b/Engine/Source/Developer/Virtualization/Virtualization.Build.cs
@@ -9,12 +9,16 @@ public class Virtualization : ModuleRules
 	{
 		PrivateDependencyModuleNames.Add("Core");
 		PrivateDependencyModuleNames.Add("CoreUObject");
+		PrivateDependencyModuleNames.Add("MessageLog");
 
-		//  Dependency for DDC2
+		// Dependency for DDC2
 		PrivateDependencyModuleNames.Add("DerivedDataCache");
 
-		// Dependencies for the Jupiter service backend
-		PrivateDependencyModuleNames.AddRange(new string[] { "SSL", "Json", "SourceControl" });
+		// Dependency for source control backend
+		PrivateDependencyModuleNames.Add("SourceControl");
+
+		// Dependencies for the Horde Storage backend
+		PrivateDependencyModuleNames.AddRange(new string[] { "SSL", "Json" });
 
 		AddEngineThirdPartyPrivateStaticDependencies(Target, "libcurl");
 		AddEngineThirdPartyPrivateStaticDependencies(Target, "OpenSSL");
