Add an early out path to when checking for the existence of 0 virtualized payloads to avoid making pointless calls to the backends.

#rnx

diff --git a/Engine/Source/Developer/Virtualization/Private/VirtualizationManager.cpp b/Engine/Source/Developer/Virtualization/Private/VirtualizationManager.cpp
--- a/Engine/Source/Developer/Virtualization/Private/VirtualizationManager.cpp
+++ b/Engine/Source/Developer/Virtualization/Private/VirtualizationManager.cpp
@@ -459,7 +459,12 @@ bool FVirtualizationManager::DoPayloadsExist(TArrayView<const FPayloadId> Ids, E
 {
 	TRACE_CPUPROFILER_EVENT_SCOPE(FVirtualizationManager::DoPayloadsExist);
 
-	OutStatuses.SetNum(Ids.Num());
+	OutStatuses.SetNum(Ids.Num()); // Make sure we set the number out statuses before we potentially early out
+
+	if (Ids.IsEmpty())
+	{
+		return true;
+	}
 
 	for (int32 Index = 0; Index < Ids.Num(); ++Index)
 	{
