Add an early out path to when attempting to push 0 payloads to the virtualization system to avoid making pointless calls to the backends.

#rnx

diff --git a/Engine/Source/Developer/Virtualization/Private/VirtualizationManager.cpp b/Engine/Source/Developer/Virtualization/Private/VirtualizationManager.cpp
--- a/Engine/Source/Developer/Virtualization/Private/VirtualizationManager.cpp
+++ b/Engine/Source/Developer/Virtualization/Private/VirtualizationManager.cpp
@@ -298,6 +298,11 @@ bool FVirtualizationManager::PushData(TArrayView<FPushRequest> Requests, EStorag
 {
 	TRACE_CPUPROFILER_EVENT_SCOPE(FVirtualizationManager::PushData);
 
+	if (Requests.IsEmpty())
+	{
+		return true;
+	}
+
 	TArray<FPushRequest> ValidatedRequests;
 	ValidatedRequests.Reserve(Requests.Num());
 
