The virtualization system now returns the correct statistics for pushing new payloads.

#rnx

- FScopedStatsCounter can only add a single hit/size in it's destructor, so a for loop calling ::AddHit on it was just overwriting the values each time.
- Use TrackCyclesOnly so that 'Timer' only adds the time taken to the stats.
- Now the for loop adds hits to the stats manually for each successful request.

diff --git a/Engine/Source/Developer/Virtualization/Private/VirtualizationManager.cpp b/Engine/Source/Developer/Virtualization/Private/VirtualizationManager.cpp
--- a/Engine/Source/Developer/Virtualization/Private/VirtualizationManager.cpp
+++ b/Engine/Source/Developer/Virtualization/Private/VirtualizationManager.cpp
@@ -755,7 +755,7 @@ void FVirtualizationManager::ParseHierarchy(const TCHAR* GraphName, const TCHAR*
 		UE_LOG(LogVirtualization, Fatal, TEXT("The '%s' entry for backend graph '%s' is empty [ini=%s]."), HierarchyKey, GraphName, *GEngineIni);
 	}
 
-	TArray<FString> Entries = ParseEntries(HierarchyData);
+	const TArray<FString> Entries = ParseEntries(HierarchyData);
 
 	UE_LOG(LogVirtualization, Log, TEXT("The backend graph hierarchy '%s' has %d entries"), HierarchyKey, Entries.Num());
 
@@ -865,7 +865,7 @@ void FVirtualizationManager::CachePayload(const FPayloadId& Id, const FCompresse
 			return; // No point going past BackendSource
 		}
 
-		bool bResult = TryCacheDataToBackend(*BackendToCache, Id, Payload);
+		const bool bResult = TryCacheDataToBackend(*BackendToCache, Id, Payload);
 		UE_CLOG(	!bResult, LogVirtualization, Warning,
 					TEXT("Failed to cache payload '%s' to backend '%s'"),
 					*Id.ToString(),
@@ -899,22 +899,32 @@ bool FVirtualizationManager::TryCacheDataToBackend(IVirtualizationBackend& Backe
 
 bool FVirtualizationManager::TryPushDataToBackend(IVirtualizationBackend& Backend, TArrayView<FPushRequest> Requests)
 {
-	COOK_STAT(FCookStats::FScopedStatsCounter Timer(Profiling::GetPushStats(Backend)));
-	const bool bResult = Backend.PushData(Requests);
+	COOK_STAT(FCookStats::CallStats & Stats = Profiling::GetPushStats(Backend));
+	COOK_STAT(FCookStats::FScopedStatsCounter Timer(Stats));
+	COOK_STAT(Timer.TrackCyclesOnly());
 
-	if (bResult)
+	const bool bPushResult = Backend.PushData(Requests);
+
+#if ENABLE_COOK_STATS
+	if (bPushResult)
 	{
-		for (FPushRequest& Request : Requests)
+		Timer.AddHit(0);
+
+		const bool bIsInGameThread = IsInGameThread();
+
+		for (const FPushRequest& Request : Requests)
 		{
 			// TODO: Don't add a hit if the payload was already uploaded
 			if (Request.Status == FPushRequest::EStatus::Success)
 			{
-				COOK_STAT(Timer.AddHit(Request.Payload.GetCompressedSize()));
+				Stats.Accumulate(FCookStats::CallStats::EHitOrMiss::Hit, FCookStats::CallStats::EStatType::Counter, 1l, bIsInGameThread);	
+				Stats.Accumulate(FCookStats::CallStats::EHitOrMiss::Hit, FCookStats::CallStats::EStatType::Bytes, Request.Payload.GetCompressedSize(), bIsInGameThread);
 			}
 		}
 	}
+#endif // ENABLE_COOK_STATS
 	
-	return bResult;
+	return bPushResult;
 }
 
 FCompressedBuffer FVirtualizationManager::PullDataFromBackend(IVirtualizationBackend& Backend, const FPayloadId& Id)
