Add LexToString to FPayloadId and change existing use of ToString to use it.

#rnx

diff --git a/Engine/Source/Developer/Virtualization/Private/DumpPackagePayloadInfoCommand.cpp b/Engine/Source/Developer/Virtualization/Private/DumpPackagePayloadInfoCommand.cpp
--- a/Engine/Source/Developer/Virtualization/Private/DumpPackagePayloadInfoCommand.cpp
+++ b/Engine/Source/Developer/Virtualization/Private/DumpPackagePayloadInfoCommand.cpp
@@ -53,7 +53,7 @@ void DumpPackagePayloadInfo(const TArray<FString>& Args)
 				UE_LOG(LogVirtualization, Display, TEXT("LocalPayloads:"));
 				for (int32 Index = 0; Index < LocalPayloadIds.Num(); ++Index)
 				{
-					UE_LOG(LogVirtualization, Display, TEXT("%02d: '%s'"), Index, *LocalPayloadIds[Index].ToString());
+					UE_LOG(LogVirtualization, Display, TEXT("%02d: '%s'"), Index, *LexToString(LocalPayloadIds[Index]));
 				}
 			}
 
@@ -62,7 +62,7 @@ void DumpPackagePayloadInfo(const TArray<FString>& Args)
 				UE_LOG(LogVirtualization, Display, TEXT("VirtualizedPayloads:"));
 				for (int32 Index = 0; Index < VirtualizedPayloadIds.Num(); ++Index)
 				{
-					UE_LOG(LogVirtualization, Display, TEXT("%02d: '%s'"), Index, *VirtualizedPayloadIds[Index].ToString());
+					UE_LOG(LogVirtualization, Display, TEXT("%02d: '%s'"), Index, *LexToString(VirtualizedPayloadIds[Index]));
 				}
 			}
 		}
diff --git a/Engine/Source/Developer/Virtualization/Private/PackageSubmissionChecks.cpp b/Engine/Source/Developer/Virtualization/Private/PackageSubmissionChecks.cpp
--- a/Engine/Source/Developer/Virtualization/Private/PackageSubmissionChecks.cpp
+++ b/Engine/Source/Developer/Virtualization/Private/PackageSubmissionChecks.cpp
@@ -237,7 +237,7 @@ void OnPrePackageSubmission(const TArray<FString>& FilesToSubmit, TArray<FText>&
 				else
 				{
 					FText Message = FText::Format(	LOCTEXT("Virtualization_UpdateStatusFailed", "Unable to update the status for the payload '{0}' in the package '{1}'"),
-													FText::FromString(PackageInfo.LocalPayloads[Index].ToString()),
+													FText::FromString(LexToString(PackageInfo.LocalPayloads[Index])),
 													FText::FromString(PackageInfo.Path.GetDebugName()));
 					Errors.Add(Message);
 					return;
@@ -293,7 +293,7 @@ void OnPrePackageSubmission(const TArray<FString>& FilesToSubmit, TArray<FText>&
 			{
 				FText Message = FText::Format(	LOCTEXT("Virtualization_WrongPayload", "Package {0} loaded an incorrect payload from the trailer. Expected '{1}' Loaded  '{2}'"),
 												FText::FromString(PackageInfo.Path.GetDebugName()),
-												FText::FromString(PayloadId.ToString()),
+												FText::FromString(LexToString(PayloadId)),
 												FText::FromString(LexToString(Payload.GetRawHash())));
 				Errors.Add(Message);
 				return;
@@ -302,7 +302,7 @@ void OnPrePackageSubmission(const TArray<FString>& FilesToSubmit, TArray<FText>&
 			if (!Payload)
 			{
 				FText Message = FText::Format(	LOCTEXT("Virtualization_MissingPayload", "Unable to find the payload '{0}' in the local storage of package '{1}'"),
-												FText::FromString(PayloadId.ToString()),
+												FText::FromString(LexToString(PayloadId)),
 												FText::FromString(PackageInfo.Path.GetDebugName()));
 				Errors.Add(Message);
 				return;
@@ -337,7 +337,7 @@ void OnPrePackageSubmission(const TArray<FString>& FilesToSubmit, TArray<FText>&
 				else
 					{
 					FText Message = FText::Format(	LOCTEXT("Virtualization_UpdateStatusFailed", "Unable to update the status for the payload '{0}' in the package '{1}'"),
-													FText::FromString(Request.Identifier.ToString()),
+													FText::FromString(LexToString(Request.Identifier)),
 													FText::FromString(PackageInfo.Path.GetDebugName()));
 					Errors.Add(Message);
 					return;
diff --git a/Engine/Source/Developer/Virtualization/Private/VirtualizationDDCBackend.cpp b/Engine/Source/Developer/Virtualization/Private/VirtualizationDDCBackend.cpp
--- a/Engine/Source/Developer/Virtualization/Private/VirtualizationDDCBackend.cpp
+++ b/Engine/Source/Developer/Virtualization/Private/VirtualizationDDCBackend.cpp
@@ -77,7 +77,7 @@ EPushResult FDDCBackend::PushData(const FPayloadId& Id, const FCompressedBuffer&
 
 	if (DoesPayloadExist(Id))
 	{
-		UE_LOG(LogVirtualization, Verbose, TEXT("[%s] Already has a copy of the payload '%s'."), *GetDebugName(), *Id.ToString());
+		UE_LOG(LogVirtualization, Verbose, TEXT("[%s] Already has a copy of the payload '%s'."), *GetDebugName(), *LexToString(Id));
 		return EPushResult::PayloadAlreadyExisted;
 	}
 
diff --git a/Engine/Source/Developer/Virtualization/Private/VirtualizationFileBackend.cpp b/Engine/Source/Developer/Virtualization/Private/VirtualizationFileBackend.cpp
--- a/Engine/Source/Developer/Virtualization/Private/VirtualizationFileBackend.cpp
+++ b/Engine/Source/Developer/Virtualization/Private/VirtualizationFileBackend.cpp
@@ -60,7 +60,7 @@ EPushResult FFileSystemBackend::PushData(const FPayloadId& Id, const FCompressed
 
 	if (DoesPayloadExist(Id))
 	{
-		UE_LOG(LogVirtualization, Verbose, TEXT("[%s] Already has a copy of the payload '%s'."), *GetDebugName(), *Id.ToString());
+		UE_LOG(LogVirtualization, Verbose, TEXT("[%s] Already has a copy of the payload '%s'."), *GetDebugName(), *LexToString(Id));
 		return EPushResult::PayloadAlreadyExisted;
 	}
 
@@ -80,7 +80,7 @@ EPushResult FFileSystemBackend::PushData(const FPayloadId& Id, const FCompressed
 
 		UE_LOG(LogVirtualization, Error, TEXT("[%s] Failed to write payload '%s' to '%s' due to system error: %s"), 
 			*GetDebugName(),
-			*Id.ToString(),
+			*LexToString(Id),
 			*TempFilePath,
 			SystemErrorMsg.ToString());
 
@@ -100,7 +100,7 @@ EPushResult FFileSystemBackend::PushData(const FPayloadId& Id, const FCompressed
 
 		UE_LOG(LogVirtualization, Error, TEXT("[%s] Failed to write payload '%s' contents to '%s' due to system error: %s"),
 			*GetDebugName(),
-			*Id.ToString(),
+			*LexToString(Id),
 			*TempFilePath,
 			SystemErrorMsg.ToString());
 
@@ -125,14 +125,14 @@ EPushResult FFileSystemBackend::PushData(const FPayloadId& Id, const FCompressed
 		// don't need to give an error message.
 		if (DoesPayloadExist(Id))
 		{
-			UE_LOG(LogVirtualization, Verbose, TEXT("[%s] Already has a copy of the payload '%s'."), *GetDebugName(), *Id.ToString());
+			UE_LOG(LogVirtualization, Verbose, TEXT("[%s] Already has a copy of the payload '%s'."), *GetDebugName(), *LexToString(Id));
 			return EPushResult::PayloadAlreadyExisted;
 		}
 		else
 		{
 			UE_LOG(	LogVirtualization, Error, TEXT("[%s] Failed to move payload '%s' to it's final location '%s' due to system error: %s"),
 					*GetDebugName(),
-					*Id.ToString(),
+					*LexToString(Id),
 					*FilePath,
 					SystemErrorMsg.ToString());
 
@@ -153,7 +153,7 @@ FCompressedBuffer FFileSystemBackend::PullData(const FPayloadId& Id)
 	// TODO: Should we allow the error severity to be configured via ini or just not report this case at all?
 	if (!IFileManager::Get().FileExists(FilePath.ToString()))
 	{
-		UE_LOG(LogVirtualization, Verbose, TEXT("[%s] Does not contain the payload '%s'"), *GetDebugName(), *Id.ToString());
+		UE_LOG(LogVirtualization, Verbose, TEXT("[%s] Does not contain the payload '%s'"), *GetDebugName(), *LexToString(Id));
 		return FCompressedBuffer();
 	}
 
@@ -166,7 +166,7 @@ FCompressedBuffer FFileSystemBackend::PullData(const FPayloadId& Id)
 
 		UE_LOG(LogVirtualization, Error, TEXT("[%s] Failed to load payload '%s' from file '%s' due to system error: %s"),
 			*GetDebugName(),
-			*Id.ToString(),
+			*LexToString(Id),
 			FilePath.ToString(),
 			SystemErrorMsg.ToString());
 
diff --git a/Engine/Source/Developer/Virtualization/Private/VirtualizationHttpBackend.cpp b/Engine/Source/Developer/Virtualization/Private/VirtualizationHttpBackend.cpp
--- a/Engine/Source/Developer/Virtualization/Private/VirtualizationHttpBackend.cpp
+++ b/Engine/Source/Developer/Virtualization/Private/VirtualizationHttpBackend.cpp
@@ -1383,7 +1383,7 @@ EPushResult FHttpBackend::PushData(const FPayloadId& Id, const FCompressedBuffer
 #if UE_CHECK_FOR_EXISTING_PAYLOADS
 	if (DoesPayloadExist(Id))
 	{
-		UE_LOG(LogVirtualization, Verbose, TEXT("HordeStorage already has a copy of the payload '%s'"), *Id.ToString());
+		UE_LOG(LogVirtualization, Verbose, TEXT("HordeStorage already has a copy of the payload '%s'"), *LexToString(Id));
 		return EPushResult::PayloadAlreadyExisted;
 	}
 #endif // UE_CHECK_FOR_EXISTING_PAYLOADS
@@ -1399,7 +1399,7 @@ EPushResult FHttpBackend::PushData(const FPayloadId& Id, const FCompressedBuffer
 	const int64 NumChunks = FMath::DivideAndRoundUp(FlattenedPayload.GetSize(), ChunkSize);
 	if (NumChunks > MAX_int32)
 	{
-		UE_LOG(LogVirtualization, Error, TEXT("Too many chunks (%d) are required for the payload '%s', try increasing the ChunkSize"), NumChunks, *Id.ToString());
+		UE_LOG(LogVirtualization, Error, TEXT("Too many chunks (%d) are required for the payload '%s', try increasing the ChunkSize"), NumChunks, *LexToString(Id));
 		return EPushResult::Failed;
 	}
 
@@ -1449,15 +1449,15 @@ EPushResult FHttpBackend::PushData(const FPayloadId& Id, const FCompressedBuffer
 
 	if (NumFailedChunks > 0)
 	{
-		UE_LOG(LogVirtualization, Error, TEXT("Failed to upload %d chunks for the payload '%s'."), NumFailedChunks.load(), *Id.ToString());
+		UE_LOG(LogVirtualization, Error, TEXT("Failed to upload %d chunks for the payload '%s'."), NumFailedChunks.load(), *LexToString(Id));
 		return EPushResult::Failed;
 	}
 
-	UE_LOG(LogVirtualization, Verbose, TEXT("Successfully uploaded all chunks for the payload '%s'"), *Id.ToString());
+	UE_LOG(LogVirtualization, Verbose, TEXT("Successfully uploaded all chunks for the payload '%s'"), *LexToString(Id));
 
 	// Note that the ddc end point is used by both ddc and mirage
 	TStringBuilder<256> Uri;
-	Uri.Appendf(TEXT("api/v1/c/ddc/%s/%s/%s"), *Namespace, *Bucket, *Id.ToString());
+	Uri.Appendf(TEXT("api/v1/c/ddc/%s/%s/%s"), *Namespace, *Bucket, *LexToString(Id));
 
 	// Retry request until we get an accepted response or exhaust allowed number of attempts.
 	uint32 Attempts = 0;
@@ -1476,19 +1476,19 @@ EPushResult FHttpBackend::PushData(const FPayloadId& Id, const FCompressedBuffer
 
 			if (ResponseCode == 200)
 			{
-				UE_LOG(LogVirtualization, Verbose, TEXT("Successfully uploaded the description for the payload '%s'"), *Id.ToString());
+				UE_LOG(LogVirtualization, Verbose, TEXT("Successfully uploaded the description for the payload '%s'"), *LexToString(Id));
 				return EPushResult::Success;
 			}
 
 			if (!ShouldRetryOnError(ResponseCode))
 			{
-				UE_LOG(LogVirtualization, Error, TEXT("Failed with error code '%d' to upload header infomation about payload '%s'"), ResponseCode, *Id.ToString());
+				UE_LOG(LogVirtualization, Error, TEXT("Failed with error code '%d' to upload header infomation about payload '%s'"), ResponseCode, *LexToString(Id));
 				return EPushResult::Failed;
 			}
 		}
 	}
 
-	UE_LOG(LogVirtualization, Error, TEXT("Failed  '%d' attempts to upload header infomation about payload '%s'"), UE_MIRAGE_MAX_ATTEMPTS, *Id.ToString());
+	UE_LOG(LogVirtualization, Error, TEXT("Failed  '%d' attempts to upload header infomation about payload '%s'"), UE_MIRAGE_MAX_ATTEMPTS, *LexToString(Id));
 	return EPushResult::Failed;
 }
 
@@ -1506,7 +1506,7 @@ FCompressedBuffer FHttpBackend::PullData(const FPayloadId& Id)
 	// the Io service, the ids will be in the correct order.
 	// fields=metadata - Ask for the payload metadata which contains info we can use later for optimizations.
 	TStringBuilder<256> Uri;
-	Uri.Appendf(TEXT("api/v1/c/ddc/%s/%s/%s.json?fields=contentHash&fields=blobIdentifiers&fields=metadata"), *Namespace, *Bucket, *Id.ToString());
+	Uri.Appendf(TEXT("api/v1/c/ddc/%s/%s/%s.json?fields=contentHash&fields=blobIdentifiers&fields=metadata"), *Namespace, *Bucket, *LexToString(Id));
 
 	FDDCCacheGETResponse Response;
 
@@ -1526,7 +1526,7 @@ FCompressedBuffer FHttpBackend::PullData(const FPayloadId& Id)
 			{
 				if (!Response.FromJson(Request->GetResponseAsJsonObject()))
 				{
-					UE_LOG(LogVirtualization, Error, TEXT("Failed to parser the header infomation about payload '%s'"), *Id.ToString());
+					UE_LOG(LogVirtualization, Error, TEXT("Failed to parser the header infomation about payload '%s'"), *LexToString(Id));
 					return FCompressedBuffer();
 				}
 			}
@@ -1535,12 +1535,12 @@ FCompressedBuffer FHttpBackend::PullData(const FPayloadId& Id)
 				// Response 400 indicates that the payload does not exist in HordeStorage. Note that it is faster to just make the request
 				// and check for the response rather than call ::DoesPayloadExist prior to requesting the json header because this way 
 				// we will only make a single request if the payload exists or not.
-				UE_LOG(LogVirtualization, Verbose, TEXT("[%s] Does not contain the payload '%s'"), *GetDebugName(), *Id.ToString());
+				UE_LOG(LogVirtualization, Verbose, TEXT("[%s] Does not contain the payload '%s'"), *GetDebugName(), *LexToString(Id));
 				return FCompressedBuffer();
 			}
 			else if (!ShouldRetryOnError(ResponseCode))
 			{
-				UE_LOG(LogVirtualization, Error, TEXT("Failed with error code '%d' to download header infomation about payload '%s'"), ResponseCode, *Id.ToString());
+				UE_LOG(LogVirtualization, Error, TEXT("Failed with error code '%d' to download header infomation about payload '%s'"), ResponseCode, *LexToString(Id));
 				return FCompressedBuffer();
 			}
 		}
@@ -1548,11 +1548,11 @@ FCompressedBuffer FHttpBackend::PullData(const FPayloadId& Id)
 
 	if (ResponseCode != 200)
 	{
-		UE_LOG(LogVirtualization, Error, TEXT("Failed '%d' attempts to download header infomation about payload (last error code '%d')  '%s'"), UE_MIRAGE_MAX_ATTEMPTS, ResponseCode, *Id.ToString());
+		UE_LOG(LogVirtualization, Error, TEXT("Failed '%d' attempts to download header infomation about payload (last error code '%d')  '%s'"), UE_MIRAGE_MAX_ATTEMPTS, ResponseCode, *LexToString(Id));
 		return FCompressedBuffer();
 	}
 
-	UE_LOG(LogVirtualization, Verbose, TEXT("Successfully downloaded a description for the payload '%s'"), *Id.ToString());
+	UE_LOG(LogVirtualization, Verbose, TEXT("Successfully downloaded a description for the payload '%s'"), *LexToString(Id));
 
 	// Now that we have the payload description we can start pulling the chunks from the Io service
 	// and reconstruct the final payload.
@@ -1599,12 +1599,12 @@ FCompressedBuffer FHttpBackend::PullData(const FPayloadId& Id)
 
 	if (NumFailedChunks == 0)
 	{
-		UE_LOG(LogVirtualization, Verbose, TEXT("Successfully downloaded all chunks for the payload '%s'"), *Id.ToString());	
+		UE_LOG(LogVirtualization, Verbose, TEXT("Successfully downloaded all chunks for the payload '%s'"), *LexToString(Id));
 		return FCompressedBuffer::FromCompressed(Payload.MoveToShared());
 	}
 	else
 	{
-		UE_LOG(LogVirtualization, Error, TEXT("Failed to download %d chunks for the payload '%s'"), NumFailedChunks.load(), *Id.ToString());
+		UE_LOG(LogVirtualization, Error, TEXT("Failed to download %d chunks for the payload '%s'"), NumFailedChunks.load(), *LexToString(Id));
 		return FCompressedBuffer();
 	}
 }
@@ -1617,7 +1617,7 @@ bool FHttpBackend::DoesPayloadExist(const FPayloadId& Id)
 
 	// Note that the ddc end point is used by both ddc and mirage
 	TStringBuilder<256> Uri;
-	Uri.Appendf(TEXT("api/v1/c/ddc/%s/%s/%s"), *Namespace, *Bucket, *Id.ToString());
+	Uri.Appendf(TEXT("api/v1/c/ddc/%s/%s/%s"), *Namespace, *Bucket, *LexToString(Id));
 
 	// Retry request until we get an accepted response or exhaust allowed number of attempts.
 	uint32 Attempts = 0;
@@ -1910,7 +1910,7 @@ bool FHttpBackend::PostChunk(const TArrayView<const uint8>& ChunkData, const FPa
 
 			if (ResponseCode == 200)
 			{
-				UE_LOG(LogVirtualization, Verbose, TEXT("Successfully uploaded a chunk '%s'for payload '%s'"), *OutHashAsString, *PayloadId.ToString());
+				UE_LOG(LogVirtualization, Verbose, TEXT("Successfully uploaded a chunk '%s'for payload '%s'"), *OutHashAsString, *LexToString(PayloadId));
 				return true;
 			}
 
@@ -1945,7 +1945,7 @@ bool FHttpBackend::PullChunk(const FString& Hash, const FPayloadId& PayloadId, u
 
 			if (Result != FRequest::Success)
 			{
-				UE_LOG(LogVirtualization, Error, TEXT("Attempting to GET a payload chunk '%s' for payload '%s' failed due to an internal Curl error"), *Hash, *PayloadId.ToString());	
+				UE_LOG(LogVirtualization, Error, TEXT("Attempting to GET a payload chunk '%s' for payload '%s' failed due to an internal Curl error"), *Hash, *LexToString(PayloadId));
 				return false;
 			}
 
@@ -1953,19 +1953,19 @@ bool FHttpBackend::PullChunk(const FString& Hash, const FPayloadId& PayloadId, u
 
 			if (ResponseCode == 200)
 			{
-				UE_LOG(LogVirtualization, Verbose, TEXT("Successfully downloaded a payload chunk '%s' for payload '%s'"), *Hash, *PayloadId.ToString());	
+				UE_LOG(LogVirtualization, Verbose, TEXT("Successfully downloaded a payload chunk '%s' for payload '%s'"), *Hash, *LexToString(PayloadId));
 				return true;
 			}
 
 			if (!ShouldRetryOnError(ResponseCode))
 			{
-				UE_LOG(LogVirtualization, Error, TEXT("Attempting to GET a payload chunk '%s' for payload '%s' failed with http response: %" INT64_FMT), *Hash, *PayloadId.ToString(), ResponseCode);	
+				UE_LOG(LogVirtualization, Error, TEXT("Attempting to GET a payload chunk '%s' for payload '%s' failed with http response: %" INT64_FMT), *Hash, *LexToString(PayloadId), ResponseCode);
 				return false;
 			}
 		}
 	}
 
-	UE_LOG(LogVirtualization, Error, TEXT("Attempting to GET a payload chunk '%s' for payload '%s' failed all '%d' attempts"), *Hash, *PayloadId.ToString(), UE_MIRAGE_MAX_ATTEMPTS);	
+	UE_LOG(LogVirtualization, Error, TEXT("Attempting to GET a payload chunk '%s' for payload '%s' failed all '%d' attempts"), *Hash, *LexToString(PayloadId), UE_MIRAGE_MAX_ATTEMPTS);
 	return false;
 }
 
diff --git a/Engine/Source/Developer/Virtualization/Private/VirtualizationManager.cpp b/Engine/Source/Developer/Virtualization/Private/VirtualizationManager.cpp
--- a/Engine/Source/Developer/Virtualization/Private/VirtualizationManager.cpp
+++ b/Engine/Source/Developer/Virtualization/Private/VirtualizationManager.cpp
@@ -324,7 +324,7 @@ bool FVirtualizationManager::PushData(TArrayView<FPushRequest> Requests, EStorag
 		if ((int64)Request.Payload.GetCompressedSize() < MinPayloadLength)
 		{
 			UE_LOG(	LogVirtualization, Verbose, TEXT("Attempting to push a virtualized payload (id: %s) that is smaller (%" UINT64_FMT ") than the MinPayloadLength (%" INT64_FMT ")"),
-					*Request.Identifier.ToString(), 
+					*LexToString(Request.Identifier),
 					Request.Payload.GetCompressedSize(),
 					MinPayloadLength);
 
@@ -384,7 +384,7 @@ bool FVirtualizationManager::PushData(TArrayView<FPushRequest> Requests, EStorag
 				checkf(	Request.Payload.GetRawHash() == ValidationPayload.GetRawHash(),
 						TEXT("[%s] Failed to pull payload '%s' after it was pushed to backend"),
 						*Backend->GetDebugName(),
-						*Request.Identifier.ToString());
+						*LexToString(Request.Identifier));
 			}
 		}
 	}
@@ -420,7 +420,7 @@ FCompressedBuffer FVirtualizationManager::PullData(const FPayloadId& Id)
 	if (PullEnabledBackends.IsEmpty())
 	{
 		// TODO: See below, should errors here be fatal?
-		UE_LOG(LogVirtualization, Error, TEXT("Payload '%s' failed to be pulled as there are no backends mounted!'"), *Id.ToString());
+		UE_LOG(LogVirtualization, Error, TEXT("Payload '%s' failed to be pulled as there are no backends mounted!'"), *LexToString(Id));
 		return FCompressedBuffer();
 	}
 
@@ -455,7 +455,7 @@ FCompressedBuffer FVirtualizationManager::PullData(const FPayloadId& Id)
 	// code handles it properly.
 	// Could be worth extending ::PullData to return error codes instead so we can make a better distinction 
 	// between the payload not being found in any of the backends and one or more of the backends failing.
-	UE_LOG(LogVirtualization, Error, TEXT("Payload '%s' failed to be pulled from any backend'"), *Id.ToString());
+	UE_LOG(LogVirtualization, Error, TEXT("Payload '%s' failed to be pulled from any backend'"), *LexToString(Id));
 
 	return FCompressedBuffer();
 }
@@ -868,7 +868,7 @@ void FVirtualizationManager::CachePayload(const FPayloadId& Id, const FCompresse
 		const bool bResult = TryCacheDataToBackend(*BackendToCache, Id, Payload);
 		UE_CLOG(	!bResult, LogVirtualization, Warning,
 					TEXT("Failed to cache payload '%s' to backend '%s'"),
-					*Id.ToString(),
+					*LexToString(Id),
 					*BackendToCache->GetDebugName());
 
 		// Debugging operation where we immediately try to pull the payload after each push (when possible) and assert 
@@ -879,7 +879,7 @@ void FVirtualizationManager::CachePayload(const FPayloadId& Id, const FCompresse
 			checkf(	Payload.GetRawHash() == PulledPayload.GetRawHash(), 
 					TEXT("[%s] Failed to pull payload '%s' after it was cached to backend"),
 					*BackendToCache->GetDebugName(), 
-					*Id.ToString());
+					*LexToString(Id));
 		}
 	}
 }
diff --git a/Engine/Source/Developer/Virtualization/Private/VirtualizationSourceControlBackend.cpp b/Engine/Source/Developer/Virtualization/Private/VirtualizationSourceControlBackend.cpp
--- a/Engine/Source/Developer/Virtualization/Private/VirtualizationSourceControlBackend.cpp
+++ b/Engine/Source/Developer/Virtualization/Private/VirtualizationSourceControlBackend.cpp
@@ -262,7 +262,7 @@ bool FSourceControlBackend::PushData(TArrayView<FPushRequest> Requests)
 
 				UE_LOG(LogVirtualization, Error, TEXT("[%s] Failed to write payload '%s' contents to '%s' due to system error: %s"),
 					*GetDebugName(),
-					*Request.Identifier.ToString(),
+					*LexToString(Request.Identifier),
 					*PayloadFilePath,
 					SystemErrorMsg.ToString());
 
@@ -278,7 +278,7 @@ bool FSourceControlBackend::PushData(TArrayView<FPushRequest> Requests)
 
 				UE_LOG(LogVirtualization, Error, TEXT("[%s] Failed to write payload '%s' contents to '%s' due to system error: %s"),
 					*GetDebugName(),
-					*Request.Identifier.ToString(),
+					*LexToString(Request.Identifier),
 					*PayloadFilePath,
 					SystemErrorMsg.
 					ToString());
diff --git a/Engine/Source/Runtime/Core/Private/Tests/Virtualization/PayloadIdTests.cpp b/Engine/Source/Runtime/Core/Private/Tests/Virtualization/PayloadIdTests.cpp
--- a/Engine/Source/Runtime/Core/Private/Tests/Virtualization/PayloadIdTests.cpp
+++ b/Engine/Source/Runtime/Core/Private/Tests/Virtualization/PayloadIdTests.cpp
@@ -310,9 +310,9 @@ bool FPayloadIdTestStringBuilder::RunTest(const FString& Parameters)
 
 	// Test again both TWideStringBuilder and TAnsiStringBuilder
 #define TEST_PAYLOAD_ID(BufferToTest) {	TWideStringBuilder<128> WideBuilder; WideBuilder << BufferToTest; \
-										TestEqual(TEXT("TWideStringBuilder << and ToString to produce the same result"), *WideBuilder, BufferToTest.ToString());\
+										TestEqual(TEXT("TWideStringBuilder << and ToString to produce the same result"), *WideBuilder, LexToString(BufferToTest));\
 										TAnsiStringBuilder<128> AnsiBuilder; AnsiBuilder << BufferToTest; \
-										TestEqual(TEXT("TAnsiBuilder << and ToString to produce the same result"), *AnsiBuilder, BufferToTest.ToString()); }
+										TestEqual(TEXT("TAnsiBuilder << and ToString to produce the same result"), *AnsiBuilder, LexToString(BufferToTest)); }
 
 	TEST_PAYLOAD_ID(EmptyId);
 	TEST_PAYLOAD_ID(BufferId0);
diff --git a/Engine/Source/Runtime/Core/Private/Virtualization/PayloadId.cpp b/Engine/Source/Runtime/Core/Private/Virtualization/PayloadId.cpp
--- a/Engine/Source/Runtime/Core/Private/Virtualization/PayloadId.cpp
+++ b/Engine/Source/Runtime/Core/Private/Virtualization/PayloadId.cpp
@@ -8,9 +8,7 @@
 #include "Misc/StringBuilder.h"
 #include "Serialization/Archive.h"
 
-namespace UE
-{
-namespace Virtualization
+namespace UE::Virtualization
 {
 
 FPayloadId::FPayloadId(const FIoHash& BlakeHash)
@@ -72,21 +70,6 @@ FGuid FPayloadId::ToGuid() const
 	}
 }
 
-FString FPayloadId::ToString() const
-{
-	if (IsValid())
-	{	
-		TStringBuilder<65> HexStringBuilder;
-		HexStringBuilder << Identifier;
-
-		return HexStringBuilder.ToString();
-	}
-	else
-	{
-		return FString();
-	}
-}
-
 bool FPayloadId::operator == (const FPayloadId& InPayloadId) const
 {
 	// If ::bIsHashSet is different between the two payloads we know they are different.
@@ -117,5 +100,9 @@ FArchive& operator<<(FArchive& Ar, FPayloadId& PayloadId)
 	return Ar;
 }
 
-} // namespace Virtualization
-} // namespace UE
+FString LexToString(const FPayloadId& Id)
+{
+	return LexToString(Id.GetIdentifier());
+}
+
+} // namespace UE::Virtualization
diff --git a/Engine/Source/Runtime/Core/Public/Virtualization/PayloadId.h b/Engine/Source/Runtime/Core/Public/Virtualization/PayloadId.h
--- a/Engine/Source/Runtime/Core/Public/Virtualization/PayloadId.h
+++ b/Engine/Source/Runtime/Core/Public/Virtualization/PayloadId.h
@@ -46,9 +46,6 @@ public:
 	 */
 	FGuid ToGuid() const;
 
-	/** Returns a string representation of the id, useful for debug tools */
-	FString ToString() const;
-
 	const FIoHash& GetIdentifier() const
 	{
 		return Identifier;
@@ -108,4 +105,6 @@ TStringBuilderBase<CharType>& operator<<(TStringBuilderBase<CharType>& Builder,
 	return Builder;
 }
 
+[[nodiscard]] CORE_API FString LexToString(const FPayloadId& Id);
+
 } // namespace UE::Virtualization
diff --git a/Engine/Source/Runtime/CoreUObject/Private/Serialization/EditorBulkData.cpp b/Engine/Source/Runtime/CoreUObject/Private/Serialization/EditorBulkData.cpp
--- a/Engine/Source/Runtime/CoreUObject/Private/Serialization/EditorBulkData.cpp
+++ b/Engine/Source/Runtime/CoreUObject/Private/Serialization/EditorBulkData.cpp
@@ -903,8 +903,8 @@ FCompressedBuffer FEditorBulkData::LoadFromDisk() const
 			UE::Virtualization::FPayloadId SidecarId(SidecarBuffer.Decompress());
 			UE::Virtualization::FPayloadId AssetId(AssetBuffer.Decompress());
 
-			UE_CLOG(SidecarId != PayloadContentId, LogSerialization, Error, TEXT("Sidecar content did not hash correctly! Found '%s' Expected '%s'"), *SidecarId.ToString(), *PayloadContentId.ToString());
-			UE_CLOG(AssetId != PayloadContentId, LogSerialization, Error, TEXT("Asset content did not hash correctly! Found '%s' Expected '%s'"), *AssetId.ToString(), *PayloadContentId.ToString())
+			UE_CLOG(SidecarId != PayloadContentId, LogSerialization, Error, TEXT("Sidecar content did not hash correctly! Found '%s' Expected '%s'"), *LexToString(SidecarId), *LexToString(PayloadContentId));
+			UE_CLOG(AssetId != PayloadContentId, LogSerialization, Error, TEXT("Asset content did not hash correctly! Found '%s' Expected '%s'"), *LexToString(AssetId), *LexToString(PayloadContentId))
 
 			return SidecarBuffer;
 		}
@@ -1057,12 +1057,12 @@ FCompressedBuffer FEditorBulkData::LoadFromSidecarFileInternal(ErrorVerbosity Ve
 			}
 			else if(Verbosity > ErrorVerbosity::None)
 			{
-				UE_LOG(LogSerialization, Error, TEXT("Payload '%s' in '%s' has an invalid OffsetInFile!"), *PayloadContentId.ToString(), *PackagePath.GetLocalFullPath(EPackageSegment::PayloadSidecar));
+				UE_LOG(LogSerialization, Error, TEXT("Payload '%s' in '%s' has an invalid OffsetInFile!"), *LexToString(PayloadContentId), *PackagePath.GetLocalFullPath(EPackageSegment::PayloadSidecar));
 			}
 		}
 		else if(Verbosity > ErrorVerbosity::None)
 		{
-			UE_LOG(LogSerialization, Error, TEXT("Unable to find payload '%s' in '%s'"), *PayloadContentId.ToString(), *PackagePath.GetLocalFullPath(EPackageSegment::PayloadSidecar));
+			UE_LOG(LogSerialization, Error, TEXT("Unable to find payload '%s' in '%s'"), *LexToString(PayloadContentId), *PackagePath.GetLocalFullPath(EPackageSegment::PayloadSidecar));
 		}
 	}
 	else if(Verbosity > ErrorVerbosity::None)
@@ -1471,8 +1471,8 @@ FCompressedBuffer FEditorBulkData::GetDataInternal() const
 		
 		checkf(Payload.IsNull(), TEXT("Pulling data somehow assigned it to the bulk data object!")); //Make sure that we did not assign the buffer internally
 
-		UE_CLOG(CompressedPayload.IsNull(), LogSerialization, Error, TEXT("Failed to pull payload '%s'"), *PayloadContentId.ToString());
-		UE_CLOG(!IsDataValid(*this, CompressedPayload), LogSerialization, UE_CORRUPTED_DATA_SEVERITY, TEXT("Virtualized payload '%s' is corrupt! Check the backend storage."), *PayloadContentId.ToString());
+		UE_CLOG(CompressedPayload.IsNull(), LogSerialization, Error, TEXT("Failed to pull payload '%s'"), *LexToString(PayloadContentId));
+		UE_CLOG(!IsDataValid(*this, CompressedPayload), LogSerialization, UE_CORRUPTED_DATA_SEVERITY, TEXT("Virtualized payload '%s' is corrupt! Check the backend storage."), *LexToString(PayloadContentId));
 		
 		return CompressedPayload;
 	}
@@ -1482,9 +1482,9 @@ FCompressedBuffer FEditorBulkData::GetDataInternal() const
 		
 		check(Payload.IsNull()); //Make sure that we did not assign the buffer internally
 
-		UE_CLOG(CompressedPayload.IsNull(), LogSerialization, Error, TEXT("Failed to load payload '%s"), *PayloadContentId.ToString());
+		UE_CLOG(CompressedPayload.IsNull(), LogSerialization, Error, TEXT("Failed to load payload '%s"), *LexToString(PayloadContentId));
 		UE_CLOG(!IsDataValid(*this, CompressedPayload), LogSerialization, UE_CORRUPTED_DATA_SEVERITY, TEXT("Payload '%s' loaded from '%s' is corrupt! Check the file on disk."),
-			*PayloadContentId.ToString(), 
+			*LexToString(PayloadContentId),
 			*PackagePath.GetDebugName());
 
 		return CompressedPayload;
