Fix the virtualization config option 'EnablePushToBackend' to no longer produce submission errors when disabled.

#rnx

- Add a way to poll the virtualization system to see if pushing to a backend storage type is possible or not.
- If we cannot push to persistent backend storage when submitting packages to source control we just log in verbose mode and return rather than giving an error. Submitting a non-virtualized package is no longer the problem it once was.

diff --git a/Engine/Source/Developer/Virtualization/Private/PackageSubmissionChecks.cpp b/Engine/Source/Developer/Virtualization/Private/PackageSubmissionChecks.cpp
--- a/Engine/Source/Developer/Virtualization/Private/PackageSubmissionChecks.cpp
+++ b/Engine/Source/Developer/Virtualization/Private/PackageSubmissionChecks.cpp
@@ -142,6 +142,12 @@ void OnPrePackageSubmission(const TArray<FString>& FilesToSubmit, TArray<FText>&
 		return;
 	}
 
+	if (!System.IsPushingEnabled(EStorageType::Persistent))
+	{
+		UE_LOG(LogVirtualization, Verbose, TEXT("Pushing to persistent backend storage is disabled"));
+		return;
+	}
+
 	const double StartTime = FPlatformTime::Seconds();
 
 	// Other systems may have added errors to this array, we need to check so later we can determine if this function added any additional errors.
diff --git a/Engine/Source/Developer/Virtualization/Private/VirtualizationManager.cpp b/Engine/Source/Developer/Virtualization/Private/VirtualizationManager.cpp
--- a/Engine/Source/Developer/Virtualization/Private/VirtualizationManager.cpp
+++ b/Engine/Source/Developer/Virtualization/Private/VirtualizationManager.cpp
@@ -288,6 +288,30 @@ bool FVirtualizationManager::IsEnabled() const
 	return !AllBackends.IsEmpty();
 }
 
+bool FVirtualizationManager::IsPushingEnabled(EStorageType StorageType) const
+{
+	if (!bEnablePayloadPushing)
+	{
+		return false;
+	}
+
+	switch (StorageType)
+	{
+		case EStorageType::Local:
+			return !LocalCachableBackends.IsEmpty();
+			break;
+
+		case EStorageType::Persistent:
+			return !PersistentStorageBackends.IsEmpty();
+			break;
+
+		default:
+			checkNoEntry();
+			return false;
+			break;
+	}
+}
+
 bool FVirtualizationManager::PushData(const FPayloadId& Id, const FCompressedBuffer& Payload, EStorageType StorageType, const FString& Context)
 {
 	FPushRequest Request(Id, Payload, Context);
diff --git a/Engine/Source/Developer/Virtualization/Private/VirtualizationManager.h b/Engine/Source/Developer/Virtualization/Private/VirtualizationManager.h
--- a/Engine/Source/Developer/Virtualization/Private/VirtualizationManager.h
+++ b/Engine/Source/Developer/Virtualization/Private/VirtualizationManager.h
@@ -92,6 +92,7 @@ private:
 	/* IVirtualizationSystem implementation */
 
 	virtual bool IsEnabled() const override;
+	virtual bool IsPushingEnabled(EStorageType StorageType) const override;
 	
 	virtual bool PushData(const FPayloadId& Id, const FCompressedBuffer& Payload, EStorageType StorageType, const FString& Context) override;
 	virtual bool PushData(TArrayView<FPushRequest> Requests, EStorageType StorageType) override;
diff --git a/Engine/Source/Runtime/Core/Private/Virtualization/VirtualizationSystem.cpp b/Engine/Source/Runtime/Core/Private/Virtualization/VirtualizationSystem.cpp
--- a/Engine/Source/Runtime/Core/Private/Virtualization/VirtualizationSystem.cpp
+++ b/Engine/Source/Runtime/Core/Private/Virtualization/VirtualizationSystem.cpp
@@ -24,6 +24,11 @@ public:
 		return false;
 	}
 
+	virtual bool IsPushingEnabled(EStorageType StorageType) const override
+	{
+		return false;
+	}
+
 	virtual bool PushData(const FPayloadId& Id, const FCompressedBuffer& Payload, EStorageType StorageType, const FString& Context) override
 	{
 		return false;
diff --git a/Engine/Source/Runtime/Core/Public/Virtualization/VirtualizationSystem.h b/Engine/Source/Runtime/Core/Public/Virtualization/VirtualizationSystem.h
--- a/Engine/Source/Runtime/Core/Public/Virtualization/VirtualizationSystem.h
+++ b/Engine/Source/Runtime/Core/Public/Virtualization/VirtualizationSystem.h
@@ -128,6 +128,9 @@ public:
 	/** Poll to see if content virtualization is enabled or not. */
 	virtual bool IsEnabled() const = 0;
 
+	/** Poll to see if pushing virtualized content to the given backend storage type is enabled or not. */
+	virtual bool IsPushingEnabled(EStorageType StorageType) const = 0;
+
 	/**
 	 * Push a payload to the virtualization backends.
 	 *
