Log how many payloads were actually uploaded to persistent backend storage when virtualizing packages.

#rnx

diff --git a/Engine/Source/Developer/Virtualization/Private/PackageSubmissionChecks.cpp b/Engine/Source/Developer/Virtualization/Private/PackageSubmissionChecks.cpp
--- a/Engine/Source/Developer/Virtualization/Private/PackageSubmissionChecks.cpp
+++ b/Engine/Source/Developer/Virtualization/Private/PackageSubmissionChecks.cpp
@@ -262,7 +262,7 @@ void OnPrePackageSubmission(const TArray<FString>& FilesToSubmit, TArray<FText>&
 		TotalPayloadsToVirtualize += PackageInfo.LocalPayloads.Num();
 	}
 
-	UE_LOG(LogVirtualization, Display, TEXT("Found %" INT64_FMT " payload(s) that need to be pushed to persistent virtualized storage"), TotalPayloadsToVirtualize);
+	UE_LOG(LogVirtualization, Display, TEXT("Found %" INT64_FMT " payload(s) that potentially need to be pushed to persistent virtualized storage"), TotalPayloadsToVirtualize);
 
 	// TODO Optimization: In theory we could have many packages sharing the same payload and we only need to push once
 	
@@ -326,6 +326,13 @@ void OnPrePackageSubmission(const TArray<FString>& FilesToSubmit, TArray<FText>&
 		return;
 	}
 
+	int64 TotalPayloadsVirtualized = 0;
+	for (const Virtualization::FPushRequest& Request : PayloadsToSubmit)
+	{
+		TotalPayloadsVirtualized += Request.Status == FPushRequest::EStatus::Success ? 1 : 0;
+	}
+	UE_LOG(LogVirtualization, Display, TEXT("Pushed %" INT64_FMT " payload(s) to persistent virtualized storage"), TotalPayloadsVirtualized);
+
 	// Update the package info for the submitted payloads
 	for (FPackageInfo& PackageInfo : Packages)
 	{
