Restore asset virtualization filtering by package path options.

#rnx

- Package filtering was disabled when the api was changed to take batches of payloads to push to the backend as the context per payload had to be changed to a raw FString rather than FPackagePath.
- We now convert each context to a FPackagePath and if it is valid test against filtering, if it is not valid we assume that the payload did not come from a package.
- Still tempted to move filtering out of virtualization manager entirely and move it to PackageSubmissionChecks, but that can always be done at a later date.
- Retested the filtering system, plus the ability to add filters to the config files for GameFeature plugins.

diff --git a/Engine/Source/Developer/Virtualization/Private/VirtualizationManager.cpp b/Engine/Source/Developer/Virtualization/Private/VirtualizationManager.cpp
--- a/Engine/Source/Developer/Virtualization/Private/VirtualizationManager.cpp
+++ b/Engine/Source/Developer/Virtualization/Private/VirtualizationManager.cpp
@@ -347,7 +347,7 @@ bool FVirtualizationManager::PushData(TArrayView<FPushRequest> Requests, EStorag
 
 		if ((int64)Request.Payload.GetCompressedSize() < MinPayloadLength)
 		{
-			UE_LOG(	LogVirtualization, Verbose, TEXT("Attempting to push a virtualized payload (id: %s) that is smaller (%" UINT64_FMT ") than the MinPayloadLength (%" INT64_FMT ")"),
+			UE_LOG(	LogVirtualization, Verbose, TEXT("Attempting to push a payload (id: %s) that is smaller (%" UINT64_FMT ") than the MinPayloadLength (%" INT64_FMT ")"),
 					*LexToString(Request.Identifier),
 					Request.Payload.GetCompressedSize(),
 					MinPayloadLength);
@@ -356,23 +356,23 @@ bool FVirtualizationManager::PushData(TArrayView<FPushRequest> Requests, EStorag
 			continue;
 		}
 
+		if (!ShouldVirtualizePackage(Request.Context))
+		{
+			UE_LOG(	LogVirtualization, Verbose, TEXT("Attempting to push a payload (id: %s) from a package path ('%s') that is excluded by filtering"),
+					*LexToString(Request.Identifier), 
+					*Request.Context);
+			
+			Request.Status = FPushRequest::EStatus::ExcludedByPackagPath;
+			continue;
+		}
+
 		OriginalToValidatedRequest[Index] = ValidatedRequests.Num();
 		ValidatedRequests.Add(Request);
 	}
-	
-	// Filtering by package name is currently disabled
-#if 0
-	if (!ShouldVirtualizePackage(PackageContext))
-	{
-		UE_LOG(LogVirtualization, Verbose, TEXT("Payload '%s' for package '%s' will not be virtualized due to filtering"),
-			*Id.ToString(), *PackageContext.GetDebugName());
-		return false;
-	}
-#endif
 
 	FConditionalScopeLock _(&ForceSingleThreadedCS, bForceSingleThreaded);
 
-	// Early out if there are no backends or if the pushing of payloads has been disabled
+	// Early out if there are no backends
 	if (!IsEnabled() || bEnablePayloadPushing == false)
 	{
 		return false;
@@ -1033,4 +1033,21 @@ bool FVirtualizationManager::ShouldVirtualizePackage(const FPackagePath& Package
 	return true;
 }
 
+bool FVirtualizationManager::ShouldVirtualizePackage(const FString& Context) const
+{
+	FPackagePath PackagePath;
+	if (FPackagePath::TryFromPackageName(Context, PackagePath))
+	{
+		return ShouldVirtualizePackage(PackagePath);
+	}
+
+	if (FPackagePath::TryFromMountedName(Context, PackagePath))
+	{
+		return ShouldVirtualizePackage(PackagePath);
+	}
+
+	// Context was not a valid package path so go ahead and virtualize it
+	return true;
+}
+
 } // namespace UE::Virtualization
diff --git a/Engine/Source/Developer/Virtualization/Private/VirtualizationManager.h b/Engine/Source/Developer/Virtualization/Private/VirtualizationManager.h
--- a/Engine/Source/Developer/Virtualization/Private/VirtualizationManager.h
+++ b/Engine/Source/Developer/Virtualization/Private/VirtualizationManager.h
@@ -141,6 +141,7 @@ private:
 	 *						excluded by the projects current filter set up.
 	 */
 	bool ShouldVirtualizePackage(const FPackagePath& PackagePath) const;
+	bool ShouldVirtualizePackage(const FString& Context) const;
 	
 	/** Are payloads allowed to be virtualized. Defaults to true. */
 	bool bEnablePayloadPushing;
diff --git a/Engine/Source/Runtime/Core/Public/Virtualization/VirtualizationSystem.h b/Engine/Source/Runtime/Core/Public/Virtualization/VirtualizationSystem.h
--- a/Engine/Source/Runtime/Core/Public/Virtualization/VirtualizationSystem.h
+++ b/Engine/Source/Runtime/Core/Public/Virtualization/VirtualizationSystem.h
@@ -64,6 +64,8 @@ struct FPushRequest
 		Invalid,
 		/** The payload is below the minimum length required for virtualization */
 		BelowMinSize,
+		/** The payload is owned by a package that is excluded from virtualization by path filtering */
+		ExcludedByPackagPath,
 		/** The payload in the request is now present in all backends */
 		Success
 	};
