Add code comments to clearly state that UE::Virtualization::FHttpBackend has no long term future and that FDDCBackend should be used instead.

#rnx

diff --git a/Engine/Source/Developer/Virtualization/Private/VirtualizationHttpBackend.h b/Engine/Source/Developer/Virtualization/Private/VirtualizationHttpBackend.h
--- a/Engine/Source/Developer/Virtualization/Private/VirtualizationHttpBackend.h
+++ b/Engine/Source/Developer/Virtualization/Private/VirtualizationHttpBackend.h
@@ -10,6 +10,10 @@ namespace UE::Utility { struct FAccessToken; }
 namespace UE::Virtualization
 {
 /**
+* NOTE: Although this backend can be used to store data directly in Horde storage, it is much better to use 
+* UE::Virtualization::FDDCBackend with a Zen enabled DDC instead. Due to this reason FHttpBackend will most
+* likely be deprecated in UE 5.1 and is only provided for experimentation purposes.
+* 
 * This backend allows data to be stored in and retrieved from the Horde storage service.
 *
 * Ini file setup:
