Virtualization backends no longer falsely report that payloads failed to push if the default code path for batched pushing is used.

#rnx

- Not all backends gain from pushing payloads in batches. In this case the backend can opt not to override the default implementation which will just iterate over the batch of payloads and then push each one individually.
- However when the default path was used, FPushRequest::Status was not updated after the push was performed, which resulted in each request returning the default (failed) status.
-- The backends that use the default path are not enabled in any of our internal projects so this was not noticed.
- So even though the payloads were pushed, the virtualization process would log and error and block any submit being performed.
-- Subsequent submits of the same packages would work because we would note much earlier that the payloads are already stored in the backend and the broken code path would never run.
- Added fall through comment as required by the coding standards

diff --git a/Engine/Source/Developer/Virtualization/Public/IVirtualizationBackend.h b/Engine/Source/Developer/Virtualization/Public/IVirtualizationBackend.h
--- a/Engine/Source/Developer/Virtualization/Public/IVirtualizationBackend.h
+++ b/Engine/Source/Developer/Virtualization/Public/IVirtualizationBackend.h
@@ -89,20 +89,24 @@ public:
 
 	virtual bool PushData(TArrayView<FPushRequest> Requests)
 	{
-		// TODO: Sort return codes
+		// TODO: Improve the error codes in the future
 		for (FPushRequest& Request : Requests)
 		{
 			EPushResult Result = PushData(Request.Identifier, Request.Payload, Request.Context);
 			switch (Result)
 			{
 			case EPushResult::Failed:
+				Request.Status = FPushRequest::EStatus::Failed;
 				return false;
 
 			case EPushResult::PayloadAlreadyExisted:
+				// falls through
 			case EPushResult::Success:
+				Request.Status = FPushRequest::EStatus::Success;
 				break;
 
 			default:
+				Request.Status = FPushRequest::EStatus::Failed;
 				checkNoEntry();
 				break;
 			}
