Allow the virtualization of assets to be opt in, as well as opt out. This would allow a project to selectively turn on virtualization for specific locations or plugins rather than for the whole project by default.

### VirtualizationManager
- Added EPackageFilterMode that can be set via [Core.ContentVirtualization] in the config file to control the default virtualization behavior.
-- When set to EPackageFilterMode::OptOut(the default) all package paths will virtualize unless excluded by a pattern in UVirtualizationFilterSettings::ExcludePackagePaths.
-- When set to EPackageFilterMode::OptIn then no package path will virtualize unless included by a pattern in UVirtualizationFilterSettings::IncludePackagePaths.
- Added a TRACE_CPUPROFILER_EVENT_SCOPE to the constructor and ::MountBackends to better track the set up time costs.
- Change use of FConfigCacheIni::LoadLocalIniFile to use GConfig, there is no need to load our own.
- Improved verbose logging to show when a payload is rejected via filtering.
- We now early out if all payloads being requested in a push are rejected during validation.
- Renamed the FString overload for ::ShouldVirtualizePackage to ::ShouldVirtualize to make the difference clearer.
- Added support for UVirtualizationFilterSettings::ExcludePackagePaths when filtering

### UVirtualizationFilterSettings
- Now has a new FString array ExcludePackagePaths, which contains the paths/patterns used to force packages to be virtualized when filtering.

diff --git a/Engine/Config/BaseEngine.ini b/Engine/Config/BaseEngine.ini
--- a/Engine/Config/BaseEngine.ini
+++ b/Engine/Config/BaseEngine.ini
@@ -1137,6 +1137,7 @@ EnablePushToBackend=False
 EnableCacheAfterPull=True
 MinPayloadLength=0
 BackendGraph=ContentVirtualizationBackendGraph_None
+FilterMode=OptOut
 FilterEngineContent=True
 FilterEnginePluginContent=True
 
diff --git a/Engine/Source/Developer/Virtualization/Private/VirtualizationFilterSettings.h b/Engine/Source/Developer/Virtualization/Private/VirtualizationFilterSettings.h
--- a/Engine/Source/Developer/Virtualization/Private/VirtualizationFilterSettings.h
+++ b/Engine/Source/Developer/Virtualization/Private/VirtualizationFilterSettings.h
@@ -15,11 +15,20 @@ class UVirtualizationFilterSettings : public UObject
 	GENERATED_UCLASS_BODY()
 
 	/**
-	 * A list of package paths that we should exclude from virtualization. 
+	 * A list of package paths that we should always prevent being virtualized .
 	 * Each path can either be a full package path to a specific package or it can be a
 	 * path to a directory, in which case all packages under that directory (and subdirectories)
-	 * will be excluded.
+	 * will be excluded from the virtualization process.
 	 */
 	UPROPERTY(Config)
 	TArray<FString> ExcludePackagePaths;
+
+	/**
+	 * A list of package paths that we should always attempt to virtualize.
+	 * Each path can either be a full package path to a specific package or it can be a
+	 * path to a directory, in which case all packages under that directory (and subdirectories)
+	 * will be included in the virtualization process.
+	 */
+	UPROPERTY(Config)
+	TArray<FString> IncludePackagePaths;
 };
diff --git a/Engine/Source/Developer/Virtualization/Private/VirtualizationManager.cpp b/Engine/Source/Developer/Virtualization/Private/VirtualizationManager.cpp
--- a/Engine/Source/Developer/Virtualization/Private/VirtualizationManager.cpp
+++ b/Engine/Source/Developer/Virtualization/Private/VirtualizationManager.cpp
@@ -53,6 +53,22 @@ private:
 	FCriticalSection* SyncObject;
 };
 
+bool LexTryParseString(EPackageFilterMode& OutValue, FStringView Buffer)
+{
+	if (Buffer == TEXT("OptOut"))
+	{
+		OutValue = EPackageFilterMode::OptOut;
+		return true;
+	}
+	else if (Buffer == TEXT("OptIn"))
+	{
+		OutValue = EPackageFilterMode::OptIn;
+		return true;
+	}
+
+	return false;
+}
+
 /* Utility function for building up a lookup table of all available IBackendFactory interfaces*/
 FVirtualizationManager::FRegistedFactories FindBackendFactories()
 {
@@ -246,6 +262,8 @@ FVirtualizationManager::FVirtualizationManager()
 	, bForceSingleThreaded(false)
 	, bValidateAfterPushOperation(false)
 {
+	TRACE_CPUPROFILER_EVENT_SCOPE(FVirtualizationManager::FVirtualizationManager);
+
 	UE_LOG(LogVirtualization, Log, TEXT("Virtualization manager created"));
 
 	// Allows us to log the profiling data on process exit. 
@@ -253,11 +271,11 @@ FVirtualizationManager::FVirtualizationManager()
 	// we need to fix the startup/shutdown ordering of Mirage first.
 	COOK_STAT(FCoreDelegates::OnExit.AddStatic(Profiling::LogStats));
 
-	FConfigFile PlatformEngineIni;
-	if (FConfigCacheIni::LoadLocalIniFile(PlatformEngineIni, TEXT("Engine"), true))
+	FConfigFile* ConfigFile = GConfig->Find(GEngineIni);
+	if (ConfigFile != nullptr)
 	{
-		ApplySettingsFromConfigFiles(PlatformEngineIni);
-		ApplyDebugSettingsFromConfigFiles(PlatformEngineIni);
+		ApplySettingsFromConfigFiles(*ConfigFile);
+		ApplyDebugSettingsFromConfigFiles(*ConfigFile);
 	}
 	else
 	{
@@ -266,7 +284,7 @@ FVirtualizationManager::FVirtualizationManager()
 
 	ApplySettingsFromCmdline();
 	ApplyDebugSettingsFromFromCmdline();
-
+	
 	MountBackends();
 }
 
@@ -356,9 +374,9 @@ bool FVirtualizationManager::PushData(TArrayView<FPushRequest> Requests, EStorag
 			continue;
 		}
 
-		if (!ShouldVirtualizePackage(Request.Context))
+		if (!ShouldVirtualize(Request.Context))
 		{
-			UE_LOG(	LogVirtualization, Verbose, TEXT("Attempting to push a payload (id: %s) from a package path ('%s') that is excluded by filtering"),
+			UE_LOG(	LogVirtualization, Verbose, TEXT("Pushing payload (id: %s) with context ('%s') was prevented by filtering"),
 					*LexToString(Request.Identifier), 
 					*Request.Context);
 			
@@ -370,6 +388,12 @@ bool FVirtualizationManager::PushData(TArrayView<FPushRequest> Requests, EStorag
 		ValidatedRequests.Add(Request);
 	}
 
+	// Early out if none of the requests require pushing after validation
+	if (ValidatedRequests.IsEmpty())
+	{
+		return true;
+	}
+
 	FConditionalScopeLock _(&ForceSingleThreadedCS, bForceSingleThreaded);
 
 	// Early out if there are no backends
@@ -651,7 +675,7 @@ void FVirtualizationManager::ApplySettingsFromConfigFiles(const FConfigFile& Pla
 	}
 	else
 	{
-		UE_LOG(LogVirtualization, Error, TEXT("Failed to load [Core.ContentVirtualization].MinPayloadLength from config file!"));;
+		UE_LOG(LogVirtualization, Error, TEXT("Failed to load [Core.ContentVirtualization].MinPayloadLength from config file!"));
 	}
 
 	FString BackendGraphNameFromIni;
@@ -662,7 +686,24 @@ void FVirtualizationManager::ApplySettingsFromConfigFiles(const FConfigFile& Pla
 	}
 	else
 	{
-		UE_LOG(LogVirtualization, Error, TEXT("Failed to load [Core.ContentVirtualization].BackendGraph from config file!"));;
+		UE_LOG(LogVirtualization, Error, TEXT("Failed to load [Core.ContentVirtualization].BackendGraph from config file!"));
+	}
+
+	FString FilterModeFromIni;
+	if (PlatformEngineIni.GetString(TEXT("Core.ContentVirtualization"), TEXT("FilterMode"), FilterModeFromIni))
+	{
+		if(LexTryParseString(FilteringMode, FilterModeFromIni))
+		{
+			UE_LOG(LogVirtualization, Log, TEXT("\tFilterMode : %s"), *FilterModeFromIni);
+		}
+		else
+		{
+			UE_LOG(LogVirtualization, Error, TEXT("[Core.ContentVirtualization].FilterMode was an invalid value! Allowed: 'OptIn'|'OptOut' Found '%s'"), *FilterModeFromIni);
+		}
+	}
+	else
+	{
+		UE_LOG(LogVirtualization, Error, TEXT("Failed to load [Core.ContentVirtualization]FilterMode from config file!"));
 	}
 
 	bool bFilterEngineContentFromIni = true;
@@ -746,6 +787,8 @@ void FVirtualizationManager::ApplyDebugSettingsFromFromCmdline()
 
 void FVirtualizationManager::MountBackends()
 {
+	TRACE_CPUPROFILER_EVENT_SCOPE(FVirtualizationManager::MountBackends);
+
 	UE_LOG(LogVirtualization, Log, TEXT("Mounting virtualization backends..."));
 
 	const FRegistedFactories FactoryLookupTable = FindBackendFactories();
@@ -1007,34 +1050,53 @@ bool FVirtualizationManager::ShouldVirtualizePackage(const FPackagePath& Package
 	const UVirtualizationFilterSettings* Settings = GetDefault<UVirtualizationFilterSettings>();
 	if (Settings != nullptr)
 	{
-		const FStringView PackageNameView = PackageName.ToView();
-
-		for (const FString& Exclusion : Settings->ExcludePackagePaths)
+		auto DoesMatch = [](const TArray<FString>& Paths, const FStringView& PackagePath) -> bool
 		{
-			if (Exclusion.EndsWith(TEXT("/")))
+			for (const FString& PathToMatch : Paths)
 			{
-				// Directory path, exclude everything under it
-				if (PackageNameView.StartsWith(Exclusion))
+				if (PathToMatch.EndsWith(TEXT("/")))
 				{
-					return false;
+					// Directory path, exclude everything under it
+					if (PackagePath.StartsWith(PathToMatch))
+					{
+						return true;
+					}
 				}
-			}
-			else
-			{
-				// Path to an asset, exclude if it matches exactly
-				if (PackageNameView == Exclusion)
+				else
 				{
-					return false;
+					// Path to an asset, exclude if it matches exactly
+					if (PackagePath == PathToMatch)
+					{
+						return true;
+					}
 				}
-			}			
+			}
+
+			return false;
+		};
+
+		const FStringView PackageNameView = PackageName.ToView();
+
+		if (DoesMatch(Settings->ExcludePackagePaths, PackageNameView))
+		{
+			return false;
+		}
+
+		if (DoesMatch(Settings->IncludePackagePaths, PackageNameView))
+		{
+			return true;
 		}
 	}
 	
-	return true;
+	// The package is not in any of the include/exclude paths so we use the default behavior
+	return ShouldVirtualizeAsDefault();
 }
 
-bool FVirtualizationManager::ShouldVirtualizePackage(const FString& Context) const
+bool FVirtualizationManager::ShouldVirtualize(const FString& Context) const
 {
+	// First see if we can convert the context from a raw string to a valid package path.
+	// If we can extract a package path then we should use the package filtering code
+	// path instead.
 	FPackagePath PackagePath;
 	if (FPackagePath::TryFromPackageName(Context, PackagePath))
 	{
@@ -1046,8 +1108,22 @@ bool FVirtualizationManager::ShouldVirtualizePackage(const FString& Context) con
 		return ShouldVirtualizePackage(PackagePath);
 	}
 
-	// Context was not a valid package path so go ahead and virtualize it
-	return true;
+	// The package is not in any of the include/exclude paths so we use the default behavior
+	return ShouldVirtualizeAsDefault();
+}
+
+bool FVirtualizationManager::ShouldVirtualizeAsDefault() const
+{
+	switch (FilteringMode)
+	{
+		case EPackageFilterMode::OptOut:
+			return true;
+		case EPackageFilterMode::OptIn:
+			return false;
+		default:
+			checkNoEntry();
+			return false;
+	}
 }
 
 } // namespace UE::Virtualization
diff --git a/Engine/Source/Developer/Virtualization/Private/VirtualizationManager.h b/Engine/Source/Developer/Virtualization/Private/VirtualizationManager.h
--- a/Engine/Source/Developer/Virtualization/Private/VirtualizationManager.h
+++ b/Engine/Source/Developer/Virtualization/Private/VirtualizationManager.h
@@ -49,26 +49,29 @@
  * Note that this only affects pushing a payload, if the filtering for a project is changed to exclude a package that
  * is already virtualized it will still be able to pull it's payloads as needed but will store them locally in the 
  * package the next time that it is saved.
- * @see ShouldVirtualizePackage for implementation details.
+ * @see ShouldVirtualizePackage or ShouldVirtualize for implementation details.
  * 
  * Basic Setup:
  * 
  * [Core.ContentVirtualization]
+ * FilterMode=OptIn/OptOut					When 'OptIn' payloads will be virtualized by default, when 'OptOut' they will not be virtualized by default
  * FilterEngineContent=True/False			When true any payload from a package under Engine/Content/.. will be excluded from virtualization
  * FilterEnginePluginContent=True/False		When true any payload from a package under Engine/Plugins/../Content/.. will be excluded from virtualization
  * 
  * PackagePath Setup:
  * 
- * The path given can either be to a directory or a specific package. It can be added to the config files for
- * a GameFeature (commonly used to exclude all content in that game feature from being virtualized) in addition 
- * to the project's config files.
+ * In addition to the default filtering mode set above, payloads stored in packages can be filtered based on the
+ * package path. This allows a package to be including in the virtualization process or excluded from it.
+ * 
  * Note that these paths will be stored in the ini files under the Saved directory. To remove a path make sure to 
  * use the - syntax to remove the entry from the array, rather than removing the line itself. Otherwise it will
  * persist until the saved config file has been reset.
  *
  * [/Script/Virtualization.VirtualizationFilterSettings]
- * +ExcludePackagePaths="/MountPoint/PathToExclude/"				Excludes any package found under '/MountPoint/PathToExclude/'
- * +ExcludePackagePaths="/MountPoint/PathTo/ThePackageToExclude"	Excludes the specific package '/MountPoint/PathTo/ThePackageToExclude'
+ * +ExcludePackagePaths="/MountPoint/PathToExclude/"				Excludes any package found under '/MountPoint/PathToExclude/' from the virtualization process
+ * +ExcludePackagePaths="/MountPoint/PathTo/ThePackageToExclude"	Excludes the specific package '/MountPoint/PathTo/ThePackageToExclude' from the virtualization process
+ * +IncludePackagePaths="/MountPoint/PathToInclude/"				Includes any package found under '/MountPoint/PathToInclude/' in the virtualization process
+ * +IncludePackagePaths="/MountPoint/PathTo/ThePackageToInclude"	Includes the specific package '/MountPoint/PathTo/ThePackageToInclude' in the virtualization process
  */
 
 namespace UE::Virtualization
@@ -76,6 +79,18 @@ namespace UE::Virtualization
 class IVirtualizationBackend;
 class IVirtualizationBackendFactory;
 
+/** The default mode of filtering to use with package paths that do not match entries in UVirtualizationFilterSettings */
+enum class EPackageFilterMode : uint8
+{
+	/** Packages will be virtualized by default and must be opted out by the use of UVirtualizationFilterSettings::ExcludePackagePaths */
+	OptOut,
+	/** Packages will not be virtualized by default and must be opted in by the user of UVirtualizationFilterSettings::IncludePackagePaths */
+	OptIn
+};
+
+/** Attempt to convert a string buffer to EPackageFilterMode */
+bool LexTryParseString(EPackageFilterMode& OutValue, FStringView Buffer);
+
 /** This is used as a wrapper around the various potential back end implementations. 
 	The calling code shouldn't need to care about which back ends are actually in use. */
 class FVirtualizationManager final : public IVirtualizationSystem
@@ -131,8 +146,10 @@ private:
 	FCompressedBuffer PullDataFromBackend(IVirtualizationBackend& Backend, const FIoHash& Id);
 
 	/** 
-	 * Determines if a package should be virtualized or not based on it's package path and the current 
-	 * filtering set up for the project.
+	 * Determines if a package path should be virtualized or not based on any exclusion/inclusion patterns
+	 * that might have been set in UVirtualizationFilterSettings.
+	 * If the path does not match any pattern set in UVirtualizationFilterSettings then use the default 
+	 * FilterMode to determine if the payload should be virtualized or not.
 	 * 
 	 * @param PackagePath	The path of the package to check. This can be empty which would indicate that
 	 *						a payload is not owned by a specific package.
@@ -140,7 +157,20 @@ private:
 	 *						excluded by the projects current filter set up.
 	 */
 	bool ShouldVirtualizePackage(const FPackagePath& PackagePath) const;
-	bool ShouldVirtualizePackage(const FString& Context) const;
+
+	/**
+	 * Determines if a package should be virtualized or not based on the given content.
+	 * If the context can be turned into a package path then ::ShouldVirtualizePackage 
+	 * will be used instead.
+	 * If the context is not a package path then we use the default FilterMode to determine
+	 * if the payload should be virtualized or not.
+	 * 
+	 * @return True if the context should be virtualized and false if not.
+	 */
+	bool ShouldVirtualize(const FString& Context) const;
+
+	/** Determines if the default filtering behavior is to virtualize a payload or not */
+	bool ShouldVirtualizeAsDefault() const;
 	
 	/** Are payloads allowed to be virtualized. Defaults to true. */
 	bool bEnablePayloadPushing;
@@ -154,6 +184,9 @@ private:
 	/** The name of the backend graph to load from the config ini file that will describe the backend hierarchy */
 	FString BackendGraphName;
 
+	/** The default filtering mode to apply if a payload is not matched with an option in UVirtualizationFilterSettings */
+	EPackageFilterMode FilteringMode = EPackageFilterMode::OptOut;
+
 	/** Should payloads in engine content packages before filtered out and never virtualized */
 	bool bFilterEngineContent;
 	
