Add a progress bar when packages are being virtualized to give feedback to the user.

#rnx

- Because all of the work is blocking on the main thread adding text to ::EnterProgressFrame to better track progress can be misleading so for now we only display "Virtualizing Assets..."

diff --git a/Engine/Source/Developer/Virtualization/Private/PackageSubmissionChecks.cpp b/Engine/Source/Developer/Virtualization/Private/PackageSubmissionChecks.cpp
--- a/Engine/Source/Developer/Virtualization/Private/PackageSubmissionChecks.cpp
+++ b/Engine/Source/Developer/Virtualization/Private/PackageSubmissionChecks.cpp
@@ -9,6 +9,7 @@
 #include "Internationalization/Internationalization.h"
 #include "Misc/PackageName.h"
 #include "Misc/Paths.h"
+#include "Misc/ScopedSlowTask.h"
 #include "Serialization/EditorBulkData.h"
 #include "UObject/Linker.h"
 #include "UObject/Package.h"
@@ -150,6 +151,9 @@ void OnPrePackageSubmission(const TArray<FString>& FilesToSubmit, TArray<FText>&
 
 	const double StartTime = FPlatformTime::Seconds();
 
+	FScopedSlowTask Progress(5.0f, LOCTEXT("Virtualization_Task", "Virtualizing Assets..."));
+	Progress.MakeDialog();
+
 	// Other systems may have added errors to this array, we need to check so later we can determine if this function added any additional errors.
 	const int32 NumErrors = Errors.Num();
 
@@ -172,6 +176,8 @@ void OnPrePackageSubmission(const TArray<FString>& FilesToSubmit, TArray<FText>&
 	TArray<FIoHash> AllLocalPayloads;
 	AllLocalPayloads.Reserve(FilesToSubmit.Num());
 
+	Progress.EnterProgressFrame(1.0f);
+
 	// From the list of files to submit we need to find all of the valid packages that contain
 	// local payloads that need to be virtualized.
 	int64 TotalPayloadsToCheck = 0;
@@ -216,6 +222,8 @@ void OnPrePackageSubmission(const TArray<FString>& FilesToSubmit, TArray<FText>&
 
 	UE_LOG(LogVirtualization, Display, TEXT("Found %" INT64_FMT " payload(s) in %d package(s) that need to be examined for virtualization"), TotalPayloadsToCheck, Packages.Num());
 
+	Progress.EnterProgressFrame(1.0f);
+
 	TArray<FPayloadStatus> PayloadStatuses;
 	if (System.QueryPayloadStatuses(AllLocalPayloads, EStorageType::Persistent, PayloadStatuses) != EQueryResult::Success)
 	{
@@ -265,7 +273,8 @@ void OnPrePackageSubmission(const TArray<FString>& FilesToSubmit, TArray<FText>&
 	UE_LOG(LogVirtualization, Display, TEXT("Found %" INT64_FMT " payload(s) that potentially need to be pushed to persistent virtualized storage"), TotalPayloadsToVirtualize);
 
 	// TODO Optimization: In theory we could have many packages sharing the same payload and we only need to push once
-	
+	Progress.EnterProgressFrame(1.0f);
+
 	TArray<Virtualization::FPushRequest> PayloadsToSubmit;
 	PayloadsToSubmit.Reserve(TotalPayloadsToVirtualize);
 
@@ -319,6 +328,8 @@ void OnPrePackageSubmission(const TArray<FString>& FilesToSubmit, TArray<FText>&
 		}
 	}
 
+	Progress.EnterProgressFrame(1.0f);
+
 	if (!System.PushData(PayloadsToSubmit, EStorageType::Persistent))
 	{
 		FText Message = LOCTEXT("Virtualization_PushFailure", "Failed to push payloads");
@@ -359,6 +370,8 @@ void OnPrePackageSubmission(const TArray<FString>& FilesToSubmit, TArray<FText>&
 		}
 	}
 
+	Progress.EnterProgressFrame(1.0f);
+
 	TArray<TPair<FPackagePath, FString>> PackagesToReplace;
 
 	// Any package with an updated trailer needs to be copied and an updated trailer appended
