Clean up the logging when initializing the virtualization module.

#rnx

- Changing a number of log items from 'Log' to 'Display' so that the info shows up more easily.
- Removed a number of log items that added no useful infomation
- Edited some of the existing log items to be clearer.

diff --git a/Engine/Source/Developer/Virtualization/Private/VirtualizationDDCBackend.cpp b/Engine/Source/Developer/Virtualization/Private/VirtualizationDDCBackend.cpp
--- a/Engine/Source/Developer/Virtualization/Private/VirtualizationDDCBackend.cpp
+++ b/Engine/Source/Developer/Virtualization/Private/VirtualizationDDCBackend.cpp
@@ -38,13 +38,13 @@ bool FDDCBackend::Initialize(const FString& ConfigEntry)
 	bool bAllowLocal = true;
 	if (FParse::Bool(*ConfigEntry, TEXT("LocalStorage="), bAllowLocal))
 	{
-		UE_LOG(LogVirtualization, Log, TEXT("[%s] Use of local storage set to '%s"), *GetDebugName(), bAllowLocal ? TEXT("true") : TEXT("false"));
+		UE_LOG(LogVirtualization, Display, TEXT("[%s] Use of local storage set to '%s"), *GetDebugName(), bAllowLocal ? TEXT("true") : TEXT("false"));
 	}
 
 	bool bAllowRemote = true;
 	if (FParse::Bool(*ConfigEntry, TEXT("RemoteStorage="), bAllowRemote))
 	{
-		UE_LOG(LogVirtualization, Log, TEXT("[%s] Use of remote storage set to '%s"), *GetDebugName(), bAllowRemote ? TEXT("true") : TEXT("false"));
+		UE_LOG(LogVirtualization, Display, TEXT("[%s] Use of remote storage set to '%s"), *GetDebugName(), bAllowRemote ? TEXT("true") : TEXT("false"));
 	}
 
 	if (!bAllowLocal && !bAllowRemote)
diff --git a/Engine/Source/Developer/Virtualization/Private/VirtualizationFileBackend.cpp b/Engine/Source/Developer/Virtualization/Private/VirtualizationFileBackend.cpp
--- a/Engine/Source/Developer/Virtualization/Private/VirtualizationFileBackend.cpp
+++ b/Engine/Source/Developer/Virtualization/Private/VirtualizationFileBackend.cpp
@@ -47,8 +47,8 @@ bool FFileSystemBackend::Initialize(const FString& ConfigEntry)
 	}
 
 	// Now log a summary of the backend settings to make issues easier to diagnose
-	UE_LOG(LogVirtualization, Log, TEXT("[%s] Using path: '%s'"), *GetDebugName(), *RootDirectory);
-	UE_LOG(LogVirtualization, Log, TEXT("[%s] Will retry failed read attempts %d times with a gap of %dms betwen them"), *GetDebugName(), RetryCount, RetryWaitTimeMS);
+	UE_LOG(LogVirtualization, Display, TEXT("[%s] Using path: '%s'"), *GetDebugName(), *RootDirectory);
+	UE_LOG(LogVirtualization, Display, TEXT("[%s] Will retry failed read attempts %d times with a gap of %dms betwen them"), *GetDebugName(), RetryCount, RetryWaitTimeMS);
 
 	return true;
 }
diff --git a/Engine/Source/Developer/Virtualization/Private/VirtualizationHttpBackend.cpp b/Engine/Source/Developer/Virtualization/Private/VirtualizationHttpBackend.cpp
--- a/Engine/Source/Developer/Virtualization/Private/VirtualizationHttpBackend.cpp
+++ b/Engine/Source/Developer/Virtualization/Private/VirtualizationHttpBackend.cpp
@@ -1242,14 +1242,14 @@ struct FHttpServiceStatus : public FJsonSerializable
 	/** Prints the status of the HordeStorage service to the log */
 	void LogStatusInfo() const
 	{
-		UE_LOG(LogVirtualization, Log, TEXT("HordeStorage Status:"));
-		UE_LOG(LogVirtualization, Log, TEXT("Version: %s"), *Version);
-		UE_LOG(LogVirtualization, Log, TEXT("Site Id: %s"), *SiteIdentifier);
-		UE_LOG(LogVirtualization, Log, TEXT("GitHash: %s"), *GitHash);		
-		UE_LOG(LogVirtualization, Log, TEXT("Capabilities:"));
+		UE_LOG(LogVirtualization, Display, TEXT("HordeStorage Status:"));
+		UE_LOG(LogVirtualization, Display, TEXT("Version: %s"), *Version);
+		UE_LOG(LogVirtualization, Display, TEXT("Site Id: %s"), *SiteIdentifier);
+		UE_LOG(LogVirtualization, Display, TEXT("GitHash: %s"), *GitHash);		
+		UE_LOG(LogVirtualization, Display, TEXT("Capabilities:"));
 		for (const FString& Capability : Capabilities)
 		{
-			UE_LOG(LogVirtualization, Log, TEXT("\t%s"), *Capability);
+			UE_LOG(LogVirtualization, Display, TEXT("\t%s"), *Capability);
 		}	
 	}
 
@@ -1322,11 +1322,11 @@ bool FHttpBackend::Initialize(const FString& ConfigEntry)
 
 	if (FParse::Value(*ConfigEntry, TEXT("ChunkSize="), ChunkSize))
 	{
-		UE_LOG(LogVirtualization, Log, TEXT("ChunkSize set to '%" UINT64_FMT "' bytes"), ChunkSize);
+		UE_LOG(LogVirtualization, Display, TEXT("ChunkSize set to '%" UINT64_FMT "' bytes"), ChunkSize);
 	}
 	else
 	{
-		UE_LOG(LogVirtualization, Log, TEXT("Payloads will not be chunked!"));
+		UE_LOG(LogVirtualization, Display, TEXT("Payloads will not be chunked!"));
 	}
 
 	// If we are connecting to a locally hosted HordeStorage then we do not need authorization
@@ -1351,7 +1351,7 @@ bool FHttpBackend::Initialize(const FString& ConfigEntry)
 		}
 	}
 
-	UE_LOG(LogVirtualization, Log, TEXT("Attempting to connect to HordeStorage at '%s' with namespace '%s'"), *HostAddress, *Namespace);
+	UE_LOG(LogVirtualization, Display, TEXT("Attempting to connect to HordeStorage at '%s' with namespace '%s'"), *HostAddress, *Namespace);
 
 	if (!IsServiceReady())
 	{
@@ -1664,7 +1664,7 @@ bool FHttpBackend::IsServiceReady() const
 
 	if (Result == FRequest::Success && FRequest::IsSuccessfulResponse(Request.GetResponseCode()))
 	{
-		UE_LOG(LogVirtualization, Log, TEXT("HordeStorage status: '%s'."), *Request.GetResponseAsString());
+		UE_LOG(LogVirtualization, Display, TEXT("HordeStorage status: '%s'."), *Request.GetResponseAsString());
 		return true;
 	}
 	else
@@ -1683,7 +1683,7 @@ bool FHttpBackend::AcquireAccessToken()
 
 	if (IsUsingLocalHost())
 	{
-		UE_LOG(LogVirtualization, Log, TEXT("Connecting to a local host '%s', so skipping authorization"), *HostAddress);
+		UE_LOG(LogVirtualization, Display, TEXT("Connecting to a local host '%s', so skipping authorization"), *HostAddress);
 		return true;
 	}
 
@@ -1764,7 +1764,7 @@ bool FHttpBackend::AcquireAccessToken()
 						AccessToken = MakeUnique<FAccessToken>();
 					}
 					AccessToken->SetHeader(*AccessTokenString);
-					UE_LOG(LogVirtualization, Log, TEXT("Logged in to HTTP DDC services. Expires in %d seconds."), ExpiryTimeSeconds);
+					UE_LOG(LogVirtualization, Display, TEXT("Logged in to HTTP DDC services. Expires in %d seconds."), ExpiryTimeSeconds);
 
 					//Schedule a refresh of the token ahead of expiry time (this will not work in commandlets)
 					if (!IsRunningCommandlet())
diff --git a/Engine/Source/Developer/Virtualization/Private/VirtualizationManager.cpp b/Engine/Source/Developer/Virtualization/Private/VirtualizationManager.cpp
--- a/Engine/Source/Developer/Virtualization/Private/VirtualizationManager.cpp
+++ b/Engine/Source/Developer/Virtualization/Private/VirtualizationManager.cpp
@@ -264,8 +264,6 @@ FVirtualizationManager::FVirtualizationManager()
 {
 	TRACE_CPUPROFILER_EVENT_SCOPE(FVirtualizationManager::FVirtualizationManager);
 
-	UE_LOG(LogVirtualization, Log, TEXT("Virtualization manager created"));
-
 	// Allows us to log the profiling data on process exit. 
 	// TODO: We should just be able to call the logging in the destructor, but 
 	// we need to fix the startup/shutdown ordering of Mirage first.
@@ -643,13 +641,13 @@ void FVirtualizationManager::GetPayloadActivityInfo( GetPayloadActivityInfoFuncR
 
 void FVirtualizationManager::ApplySettingsFromConfigFiles(const FConfigFile& PlatformEngineIni)
 {
-	UE_LOG(LogVirtualization, Log, TEXT("Loading virtualization manager settings from config files..."));
+	UE_LOG(LogVirtualization, Display, TEXT("Loading virtualization manager settings from config files..."));
 	
 	bool bEnablePayloadPushingFromIni = false;
 	if (PlatformEngineIni.GetBool(TEXT("Core.ContentVirtualization"), TEXT("EnablePushToBackend"), bEnablePayloadPushingFromIni))
 	{
 		bEnablePayloadPushing = bEnablePayloadPushingFromIni;
-		UE_LOG(LogVirtualization, Log, TEXT("\tEnablePushToBackend : %s"), bEnablePayloadPushing ? TEXT("true") : TEXT("false") );
+		UE_LOG(LogVirtualization, Display, TEXT("\tEnablePushToBackend : %s"), bEnablePayloadPushing ? TEXT("true") : TEXT("false") );
 	}
 	else
 	{
@@ -660,7 +658,7 @@ void FVirtualizationManager::ApplySettingsFromConfigFiles(const FConfigFile& Pla
 	if (PlatformEngineIni.GetBool(TEXT("Core.ContentVirtualization"), TEXT("EnableCacheAfterPull"), bEnableCacheAfterPullFromIni))
 	{
 		bEnableCacheAfterPull = bEnableCacheAfterPullFromIni;
-		UE_LOG(LogVirtualization, Log, TEXT("\tCachePulledPayloads : %s"), bEnableCacheAfterPull ? TEXT("true") : TEXT("false"));
+		UE_LOG(LogVirtualization, Display, TEXT("\tCachePulledPayloads : %s"), bEnableCacheAfterPull ? TEXT("true") : TEXT("false"));
 	}
 	else
 	{
@@ -671,7 +669,7 @@ void FVirtualizationManager::ApplySettingsFromConfigFiles(const FConfigFile& Pla
 	if (PlatformEngineIni.GetInt64(TEXT("Core.ContentVirtualization"), TEXT("MinPayloadLength"), MinPayloadLengthFromIni))
 	{
 		MinPayloadLength = MinPayloadLengthFromIni;
-		UE_LOG(LogVirtualization, Log, TEXT("\tMinPayloadLength : %" INT64_FMT), MinPayloadLength );
+		UE_LOG(LogVirtualization, Display, TEXT("\tMinPayloadLength : %" INT64_FMT), MinPayloadLength );
 	}
 	else
 	{
@@ -682,7 +680,7 @@ void FVirtualizationManager::ApplySettingsFromConfigFiles(const FConfigFile& Pla
 	if (PlatformEngineIni.GetString(TEXT("Core.ContentVirtualization"), TEXT("BackendGraph"), BackendGraphNameFromIni))
 	{
 		BackendGraphName = BackendGraphNameFromIni;
-		UE_LOG(LogVirtualization, Log, TEXT("\tBackendGraphName : %s"), *BackendGraphName );
+		UE_LOG(LogVirtualization, Display, TEXT("\tBackendGraphName : %s"), *BackendGraphName );
 	}
 	else
 	{
@@ -694,7 +692,7 @@ void FVirtualizationManager::ApplySettingsFromConfigFiles(const FConfigFile& Pla
 	{
 		if(LexTryParseString(FilteringMode, FilterModeFromIni))
 		{
-			UE_LOG(LogVirtualization, Log, TEXT("\tFilterMode : %s"), *FilterModeFromIni);
+			UE_LOG(LogVirtualization, Display, TEXT("\tFilterMode : %s"), *FilterModeFromIni);
 		}
 		else
 		{
@@ -710,7 +708,7 @@ void FVirtualizationManager::ApplySettingsFromConfigFiles(const FConfigFile& Pla
 	if (PlatformEngineIni.GetBool(TEXT("Core.ContentVirtualization"), TEXT("FilterEngineContent"), bFilterEngineContentFromIni))
 	{
 		bFilterEngineContent = bFilterEngineContentFromIni;
-		UE_LOG(LogVirtualization, Log, TEXT("\tFilterEngineContent : %s"), bFilterEngineContent ? TEXT("true") : TEXT("false"));
+		UE_LOG(LogVirtualization, Display, TEXT("\tFilterEngineContent : %s"), bFilterEngineContent ? TEXT("true") : TEXT("false"));
 	}
 	else
 	{
@@ -721,7 +719,7 @@ void FVirtualizationManager::ApplySettingsFromConfigFiles(const FConfigFile& Pla
 	if (PlatformEngineIni.GetBool(TEXT("Core.ContentVirtualization"), TEXT("FilterEnginePluginContent"), bFilterEnginePluginContentFromIni))
 	{
 		bFilterEnginePluginContent = bFilterEnginePluginContentFromIni;
-		UE_LOG(LogVirtualization, Log, TEXT("\tFilterEnginePluginContent : %s"), bFilterEnginePluginContent ? TEXT("true") : TEXT("false"));
+		UE_LOG(LogVirtualization, Display, TEXT("\tFilterEnginePluginContent : %s"), bFilterEnginePluginContent ? TEXT("true") : TEXT("false"));
 	}
 	else
 	{
@@ -734,34 +732,34 @@ void FVirtualizationManager::ApplySettingsFromCmdline()
 	FString CmdlineGraphName;
 	if (FParse::Value(FCommandLine::Get(), TEXT("-BackendGraph="), CmdlineGraphName))
 	{
-		UE_LOG(LogVirtualization, Log, TEXT("Backend graph overriden from the cmdline: '%s'"), *CmdlineGraphName);
+		UE_LOG(LogVirtualization, Display, TEXT("Backend graph overriden from the cmdline: '%s'"), *CmdlineGraphName);
 		BackendGraphName = CmdlineGraphName;
 	}
 
 	if (FParse::Param(FCommandLine::Get(), TEXT("VirtualizationForceSingleThreaded")))
 	{
 		bForceSingleThreaded = true;
-		UE_LOG(LogVirtualization, Log, TEXT("ForceSingleThreaded overriden from the cmdline: true"));	
+		UE_LOG(LogVirtualization, Display, TEXT("ForceSingleThreaded overriden from the cmdline: true"));
 	}
 }
 
 void FVirtualizationManager::ApplyDebugSettingsFromConfigFiles(const FConfigFile& PlatformEngineIni)
 {
-	UE_LOG(LogVirtualization, Log, TEXT("Loading virtualization manager debugging settings from config files..."));
+	UE_LOG(LogVirtualization, Display, TEXT("Loading virtualization manager debugging settings from config files..."));
 
 	// Note that the debug settings are optional and could be left out of the config files entirely
 	bool bForceSingleThreadedFromIni = false;
 	if (PlatformEngineIni.GetBool(TEXT("Core.ContentVirtualizationDebugOptions"), TEXT("ForceSingleThreaded"), bForceSingleThreadedFromIni))
 	{
 		bForceSingleThreaded = bForceSingleThreadedFromIni;
-		UE_LOG(LogVirtualization, Log, TEXT("\tForceSingleThreaded : %s"), bForceSingleThreaded ? TEXT("true") : TEXT("false"));
+		UE_LOG(LogVirtualization, Display, TEXT("\tForceSingleThreaded : %s"), bForceSingleThreaded ? TEXT("true") : TEXT("false"));
 	}
 
 	bool bValidateAfterPushOperationFromIni = false;
 	if (PlatformEngineIni.GetBool(TEXT("Core.ContentVirtualizationDebugOptions"), TEXT("ValidateAfterPushOperation"), bValidateAfterPushOperationFromIni))
 	{
 		bValidateAfterPushOperation = bValidateAfterPushOperationFromIni;
-		UE_LOG(LogVirtualization, Log, TEXT("ValidateAfterPushOperation : %s"), bValidateAfterPushOperation ? TEXT("true") : TEXT("false"));
+		UE_LOG(LogVirtualization, Display, TEXT("\tValidateAfterPushOperation : %s"), bValidateAfterPushOperation ? TEXT("true") : TEXT("false"));
 	}
 
 	// Some debug options will cause intentional breaks or slow downs for testing purposes, if these are enabled then we should give warning/errors 
@@ -789,8 +787,6 @@ void FVirtualizationManager::MountBackends()
 {
 	TRACE_CPUPROFILER_EVENT_SCOPE(FVirtualizationManager::MountBackends);
 
-	UE_LOG(LogVirtualization, Log, TEXT("Mounting virtualization backends..."));
-
 	const FRegistedFactories FactoryLookupTable = FindBackendFactories();
 	UE_LOG(LogVirtualization, Verbose, TEXT("Found %d backend factories"), FactoryLookupTable.Num());
 
@@ -801,7 +797,7 @@ void FVirtualizationManager::MountBackends()
 		UE_LOG(LogVirtualization, Fatal, TEXT("Unable to find the backend graph: '%s' [ini=%s]."), GraphName, *GEngineIni);
 	}
 
-	UE_LOG(LogVirtualization, Log, TEXT("Using backend graph: '%s'"), GraphName);
+	UE_LOG(LogVirtualization, Display, TEXT("Mounting virtualization backend graph: '%s'"), GraphName);
 
 	// It is important to parse the local storage hierarchy first so those backends will show up before the
 	// persistent storage backends in 'PullEnabledBackends'.
@@ -824,7 +820,7 @@ void FVirtualizationManager::ParseHierarchy(const TCHAR* GraphName, const TCHAR*
 
 	const TArray<FString> Entries = ParseEntries(HierarchyData);
 
-	UE_LOG(LogVirtualization, Log, TEXT("The backend graph hierarchy '%s' has %d entries"), HierarchyKey, Entries.Num());
+	UE_LOG(LogVirtualization, Display, TEXT("'%s' has %d backend(s)"), HierarchyKey, Entries.Num());
 
 	for (const FString& Entry : Entries)
 	{
@@ -836,7 +832,7 @@ bool FVirtualizationManager::CreateBackend(const TCHAR* GraphName, const FString
 {
 	// All failures in this method are considered fatal, however it still returns true/false in case we decide
 	// to be more forgiving in the future.
-	UE_LOG(LogVirtualization, Log, TEXT("Attempting to create back end entry '%s'"), *ConfigEntryName);
+	UE_LOG(LogVirtualization, Display, TEXT("Mounting backend entry '%s'"), *ConfigEntryName);
 
 	FString BackendData;
 	if (!GConfig->GetString(GraphName, *ConfigEntryName, BackendData, GEngineIni))
@@ -916,8 +912,6 @@ void FVirtualizationManager::AddBackend(TUniquePtr<IVirtualizationBackend> Backe
 	}
 
 	COOK_STAT(Profiling::CreateStats(*BackendRef));
-
-	UE_LOG(LogVirtualization, Log, TEXT("Mounted backend: %s"), *BackendRef->GetDebugName());
 }
 
 void FVirtualizationManager::CachePayload(const FIoHash& Id, const FCompressedBuffer& Payload, const IVirtualizationBackend* BackendSource)
diff --git a/Engine/Source/Developer/Virtualization/Private/VirtualizationSourceControlBackend.cpp b/Engine/Source/Developer/Virtualization/Private/VirtualizationSourceControlBackend.cpp
--- a/Engine/Source/Developer/Virtualization/Private/VirtualizationSourceControlBackend.cpp
+++ b/Engine/Source/Developer/Virtualization/Private/VirtualizationSourceControlBackend.cpp
@@ -67,7 +67,7 @@ bool FSourceControlBackend::Initialize(const FString& ConfigEntry)
 
 	// Optional config values
 	FParse::Bool(*ConfigEntry, TEXT("UsePartitionedClient="), bUsePartitionedClient);
-	UE_LOG(LogVirtualization, Log, TEXT("[%s] Using partitioned clients: '%s'"), *GetDebugName(), bUsePartitionedClient ? TEXT("true") : TEXT("false"));
+	UE_LOG(LogVirtualization, Display, TEXT("[%s] Using partitioned clients: '%s'"), *GetDebugName(), bUsePartitionedClient ? TEXT("true") : TEXT("false"));
 
 	ISourceControlModule& SSCModule = ISourceControlModule::Get();
 
diff --git a/Engine/Source/Runtime/Core/Private/Virtualization/VirtualizationSystem.cpp b/Engine/Source/Runtime/Core/Private/Virtualization/VirtualizationSystem.cpp
--- a/Engine/Source/Runtime/Core/Private/Virtualization/VirtualizationSystem.cpp
+++ b/Engine/Source/Runtime/Core/Private/Virtualization/VirtualizationSystem.cpp
@@ -81,7 +81,7 @@ void Initialize()
 		if (EngineIni.GetString(TEXT("Core.ContentVirtualization"), TEXT("SystemName"), RawSystemName))
 		{
 			SystemName = FName(RawSystemName);
-			UE_LOG(LogVirtualization, Log, TEXT("VirtualizationSystem name found in ini file: %s"), *RawSystemName);
+			UE_LOG(LogVirtualization, Display, TEXT("VirtualizationSystem name found in ini file: %s"), *RawSystemName);
 		}
 	}
 
