Fix the MinPayloadLength config option for asset virtualization.

#rnx

- Turns out that the option did work fine, it just needed to be tested properly, so removed the assert that was guarding against mixed package trailers when submitting virtualized assets.
- Improved verbose logging when a payload is filtered based on payload size.

diff --git a/Engine/Source/Developer/Virtualization/Private/PackageSubmissionChecks.cpp b/Engine/Source/Developer/Virtualization/Private/PackageSubmissionChecks.cpp
--- a/Engine/Source/Developer/Virtualization/Private/PackageSubmissionChecks.cpp
+++ b/Engine/Source/Developer/Virtualization/Private/PackageSubmissionChecks.cpp
@@ -359,7 +359,7 @@ void OnPrePackageSubmission(const TArray<FString>& FilesToSubmit, TArray<FText>&
 					PackageInfo.bWasTrailerUpdated = true;
 				}
 				else
-					{
+				{
 					FText Message = FText::Format(	LOCTEXT("Virtualization_UpdateStatusFailed", "Unable to update the status for the payload '{0}' in the package '{1}'"),
 													FText::FromString(LexToString(Request.Identifier)),
 													FText::FromString(PackageInfo.Path.GetDebugName()));
@@ -382,13 +382,6 @@ void OnPrePackageSubmission(const TArray<FString>& FilesToSubmit, TArray<FText>&
 			continue;
 		}
 
-		// TODO: Implement/test mixed package trailer support
-		// At the moment we assume that all the payloads in a package will virtualized or none of them will. If none of the payloads have virtualized then
-		// the package won't have made it this far so we should check that no package at this point has a local payload.
-		// We should be able to support mixed packages (to support more advanced filtering and min payload sizes etc) but it hasn't been tested yet so I'd
-		// rather prevent this condition from occurring until it is properly tested.
-		check(PackageInfo.Trailer.GetNumPayloads(EPayloadFilter::Local) == 0);
-
 		const FPackagePath& PackagePath = PackageInfo.Path; // No need to validate path, we checked this earlier
 
 		const FString PackageFilePath = PackagePath.GetLocalFullPath();
diff --git a/Engine/Source/Developer/Virtualization/Private/VirtualizationManager.cpp b/Engine/Source/Developer/Virtualization/Private/VirtualizationManager.cpp
--- a/Engine/Source/Developer/Virtualization/Private/VirtualizationManager.cpp
+++ b/Engine/Source/Developer/Virtualization/Private/VirtualizationManager.cpp
@@ -363,8 +363,9 @@ bool FVirtualizationManager::PushData(TArrayView<FPushRequest> Requests, EStorag
 
 		if ((int64)Request.Payload.GetCompressedSize() < MinPayloadLength)
 		{
-			UE_LOG(	LogVirtualization, Verbose, TEXT("Attempting to push a payload (id: %s) that is smaller (%" UINT64_FMT ") than the MinPayloadLength (%" INT64_FMT ")"),
+			UE_LOG(	LogVirtualization, Verbose, TEXT("Pushing payload (id: %s) with context ('%s') was prevented as it is smaller (%" UINT64_FMT ") than the MinPayloadLength (%" INT64_FMT ")"),
 					*LexToString(Request.Identifier),
+					*Request.Context,
 					Request.Payload.GetCompressedSize(),
 					MinPayloadLength);
 
