Cleaning up the naming conventions when creating perforce workspaces when submitting virtualized payloads

#rnx

- Workspaces created to submit virtualized assets are now prefixed as "VASubmission-" rather than "MirageSubmission-"
-- Use of the term mirage is deprecated and should be replaced by VA or VirtualizedAsset(s)
- Set the description when we create the workspace to info the user that it was auto generated during the submission of virtualized payloads.

diff --git a/Engine/Source/Developer/Virtualization/Private/VirtualizationSourceControlBackend.cpp b/Engine/Source/Developer/Virtualization/Private/VirtualizationSourceControlBackend.cpp
--- a/Engine/Source/Developer/Virtualization/Private/VirtualizationSourceControlBackend.cpp
+++ b/Engine/Source/Developer/Virtualization/Private/VirtualizationSourceControlBackend.cpp
@@ -293,7 +293,7 @@ bool FSourceControlBackend::PushData(TArrayView<FPushRequest> Requests)
 	check(Requests.Num() == FilesToSubmit.Num());
 
 	TStringBuilder<64> WorkspaceName;
-	WorkspaceName << TEXT("MirageSubmission-") << SessionGuid;
+	WorkspaceName << TEXT("VASubmission-") << SessionGuid;
 
 	ISourceControlProvider& SCCProvider = ISourceControlModule::Get().GetProvider();
 
@@ -315,6 +315,8 @@ bool FSourceControlBackend::PushData(TArrayView<FPushRequest> Requests)
 			CreateWorkspaceCommand->SetType(FCreateWorkspace::EType::Partitioned);
 		}
 
+		CreateWorkspaceCommand->SetDescription(TEXT("This workspace was autogenerated when submitting virtualized payloads to source control"));
+
 		if (SCCProvider.Execute(CreateWorkspaceCommand) != ECommandResult::Succeeded)
 		{
 			UE_LOG(LogVirtualization, Error, TEXT("[%s] Failed to create temp workspace '%s' to submit payloads from"),
