Fix some invalid LOCTEXT entries.

#rnx

- Strings inside of a LOCTEXT macro must be on a single line.

diff --git a/Engine/Source/Developer/Virtualization/Private/VirtualizationSourceControlBackend.cpp b/Engine/Source/Developer/Virtualization/Private/VirtualizationSourceControlBackend.cpp
--- a/Engine/Source/Developer/Virtualization/Private/VirtualizationSourceControlBackend.cpp
+++ b/Engine/Source/Developer/Virtualization/Private/VirtualizationSourceControlBackend.cpp
@@ -103,9 +103,7 @@ bool FSourceControlBackend::Initialize(const FString& ConfigEntry)
 		}
 
 		FMessageLog Log("LogVirtualization");
-		Log.Error(	FText::Format(LOCTEXT("FailedSourceControlConnection", "Failed to connect to source control backend with the following errors:\n{0}\n"
-					"The source control backend will be unable to pull payloads!\n"
-					"Trying logging in with the 'p4 login' command or by using p4vs/UnrealGameSync."),
+		Log.Error(	FText::Format(LOCTEXT("FailedSourceControlConnection", "Failed to connect to source control backend with the following errors:\n{0}\nThe source control backend will be unable to pull payloads!\nTrying logging in with the 'p4 login' command or by using p4vs/UnrealGameSync."),
 					Errors.ToText()));
 		
 		OnConnectionError();
@@ -122,8 +120,7 @@ bool FSourceControlBackend::Initialize(const FString& ConfigEntry)
 	{
 		FMessageLog Log("LogVirtualization");
 
-		Log.Error(	FText::Format(LOCTEXT("FailedMetaInfo", "Failed to find 'payload_metainfo.txt' in the depot '{0}'\n"
-					"The source control backend will be unable to pull payloads, is your source control  config set up correctly?"),
+		Log.Error(	FText::Format(LOCTEXT("FailedMetaInfo", "Failed to find 'payload_metainfo.txt' in the depot '{0}'\nThe source control backend will be unable to pull payloads, is your source control  config set up correctly?"),
 					FText::FromString(DepotRoot)));
 		
 		OnConnectionError();
@@ -135,8 +132,7 @@ bool FSourceControlBackend::Initialize(const FString& ConfigEntry)
 	{
 		FMessageLog Log("LogVirtualization");
 
-		Log.Error(	FText::Format(LOCTEXT("FailedMetaInfo", "Failed to find 'payload_metainfo.txt' in the depot '{0}'\n"
-					"The source control backend will be unable to pull payloads, is your source control  config set up correctly?"),
+		Log.Error(	FText::Format(LOCTEXT("FailedMetaInfo", "Failed to find 'payload_metainfo.txt' in the depot '{0}'\nThe source control backend will be unable to pull payloads, is your source control  config set up correctly?"),
 					FText::FromString(DepotRoot)));
 		
 		OnConnectionError();
@@ -149,8 +145,7 @@ bool FSourceControlBackend::Initialize(const FString& ConfigEntry)
 	{
 		FMessageLog Log("LogVirtualization");
 
-		Log.Error(	FText::Format(LOCTEXT("FailedMetaInfo", "Failed to find 'payload_metainfo.txt' in the depot '{0}'\n"
-					"The source control backend will be unable to pull payloads, is your source control  config set up correctly?"),
+		Log.Error(	FText::Format(LOCTEXT("FailedMetaInfo", "Failed to find 'payload_metainfo.txt' in the depot '{0}'\nThe source control backend will be unable to pull payloads, is your source control  config set up correctly?"),
 					FText::FromString(DepotRoot)));
 
 		OnConnectionError();
