Payload files now use the extension .upayload rather than .payload.

diff --git a/Engine/Source/Developer/Virtualization/Private/VirtualizationUtilities.cpp b/Engine/Source/Developer/Virtualization/Private/VirtualizationUtilities.cpp
--- a/Engine/Source/Developer/Virtualization/Private/VirtualizationUtilities.cpp
+++ b/Engine/Source/Developer/Virtualization/Private/VirtualizationUtilities.cpp
@@ -20,7 +20,7 @@ void PayloadIdToPath(const FIoHash& Id, FStringBuilderBase& OutPath)
 
 	OutPath.ReplaceAt(0, 6, Directory);
 
-	OutPath << TEXT(".payload");
+	OutPath << TEXT(".upayload");
 }
 
 FString PayloadIdToPath(const FIoHash& Id)
