Add a verbose log detailing the full depot path when trying to pull a virtualized payload from source control.

#rnx

diff --git a/Engine/Source/Developer/Virtualization/Private/VirtualizationSourceControlBackend.cpp b/Engine/Source/Developer/Virtualization/Private/VirtualizationSourceControlBackend.cpp
--- a/Engine/Source/Developer/Virtualization/Private/VirtualizationSourceControlBackend.cpp
+++ b/Engine/Source/Developer/Virtualization/Private/VirtualizationSourceControlBackend.cpp
@@ -167,6 +167,8 @@ FCompressedBuffer FSourceControlBackend::PullData(const FIoHash& Id)
 	TStringBuilder<512> DepotPath;
 	CreateDepotPath(Id, DepotPath);
 
+	UE_LOG(LogVirtualization, Verbose, TEXT("[%s] Attempting to pull '%s' from source control"), *GetDebugName(), *DepotPath);
+
 	ISourceControlProvider& SCCProvider = ISourceControlModule::Get().GetProvider();
 
 #if IS_SOURCE_CONTROL_THREAD_SAFE
