Improved the code logic for where we submit virtualized payloads from and added some additional options.

- Now when payloads are submitted, we will submit them from within the project's saved directory (by default) this now works because we write a .p4ignore file to the submission directory which will override any additional ignore file higher up in the directory structure.
-- A new utility function ::CreateSubmissionSessionDirectory was added to make sure that the directory we are submitting from has a p4ignore file. It is added in it's own scope as it is assumed that additional functionality will be added here later.
- If a user wishes to choose a different location to submit from, they can override the environment variable "UE-VirtualizationWorkingDir" to provide a different location. This might be useful if the user does not have a huge amount of spare space on the disk that the project is running from.
- There is also an optional entry to the ini file set up 'SubmitFromTempDir' (false by default) which when set to true will cause the system to attempt to submit the payload files from the machines temp directory instead.
-- There isn't an easy way for users to override this locally, which reduces the usefulness of the option. This might be worth improving in the future although that time would be better spent removing the need for the directory in the first place.
- We now try to find/create the root directory to submit payloads from when the backend is created. Failure at this point is not expected and will prevent the backend from working and so is considered an error.
- Added more logging to track the session guid when submitting and to be clear about where the files are being submitted from.

diff --git a/Engine/Source/Developer/Virtualization/Private/VirtualizationSourceControlBackend.cpp b/Engine/Source/Developer/Virtualization/Private/VirtualizationSourceControlBackend.cpp
--- a/Engine/Source/Developer/Virtualization/Private/VirtualizationSourceControlBackend.cpp
+++ b/Engine/Source/Developer/Virtualization/Private/VirtualizationSourceControlBackend.cpp
@@ -9,7 +9,9 @@
 #include "ISourceControlProvider.h"
 #include "Logging/MessageLog.h"
 #include "Misc/App.h"
+#include "Misc/FileHelper.h"
 #include "Misc/Parse.h"
+#include "Misc/PathViews.h"
 #include "Misc/Paths.h"
 #include "Misc/ScopeExit.h"
 #include "SourceControlOperations.h"
@@ -25,8 +27,38 @@
 namespace UE::Virtualization
 {
 
+/** Utility function to create a directory to submit payloads from. */
+[[nodiscard]] static bool TryCreateSubmissionSessionDirectory(FStringView SessionDirectoryPath)
+{
+	// Write out an ignore file to the submission directory (will create the directory if needed)
+	{
+		TStringBuilder<260> IgnoreFilePath;
+
+		// TODO: We should find if P4IGNORE is actually set and if so extract the filename to use.
+		// This will require extending the source control module
+		FPathViews::Append(IgnoreFilePath, SessionDirectoryPath, TEXT(".p4ignore.txt"));
+
+		// A very basic .p4ignore file that should make sure that we are only submitting valid .upayload files.
+		// 
+		// Since the file should only exist while we are pushing payloads, it is not expected that anyone will need
+		// to read the file. Due to this we only include the bare essentials in terms of documentation.
+
+		TStringBuilder<512> FileContents;
+
+		FileContents << TEXT("# Ignore all files\n*\n\n");
+		FileContents << TEXT("# Allow.payload files as long as they are the expected 3 directories deep\n!*/*/*/*.upayload\n\n");
+
+		if (!FFileHelper::SaveStringToFile(FileContents, IgnoreFilePath.ToString()))
+		{
+			return false;
+		}
+	}
+
+	return true;
+}
+
 /** Builds a changelist description to be used when submitting a payload to source control */
-void CreateDescription(const TArray<const FPushRequest*>& FileRequests, TStringBuilder<512>& OutDescription)
+static void CreateDescription(const TArray<const FPushRequest*>& FileRequests, TStringBuilder<512>& OutDescription)
 {
 	// TODO: Maybe make writing out the project name an option or allow for a codename to be set via ini file?
 	OutDescription << TEXT("Submitted for project: ");
@@ -65,6 +97,11 @@ bool FSourceControlBackend::Initialize(const FString& ConfigEntry)
 		return false;
 	}
 
+	if (!FindSubmissionWorkingDir(ConfigEntry))
+	{
+		return false;
+	}
+
 	// Optional config values
 	FParse::Bool(*ConfigEntry, TEXT("UsePartitionedClient="), bUsePartitionedClient);
 	UE_LOG(LogVirtualization, Display, TEXT("[%s] Using partitioned clients: '%s'"), *GetDebugName(), bUsePartitionedClient ? TEXT("true") : TEXT("false"));
@@ -227,19 +264,31 @@ bool FSourceControlBackend::PushData(TArrayView<FPushRequest> Requests)
 	// to avoid potentially conflicting with other editor processes that might be running.
 
 	const FGuid SessionGuid = FGuid::NewGuid();
-	TStringBuilder<260> RootDirectory;
-	RootDirectory << FPlatformProcess::UserTempDir() << TEXT("UnrealEngine/VirtualizedPayloads/") << SessionGuid << TEXT("/");
+	
+	UE_LOG(LogVirtualization, Display, TEXT("[%s] Started payload submission session '%s' for '%d' payload(s)"), *GetDebugName(), *LexToString(SessionGuid), Requests.Num());
+
+	TStringBuilder<260> SessionDirectory;
+	FPathViews::Append(SessionDirectory, SubmissionRootDir, SessionGuid);
+
+	if (!TryCreateSubmissionSessionDirectory(SessionDirectory))
+	{
+		UE_LOG(LogVirtualization, Error, TEXT("[%s] Failed to created directory '%s' to submit payloads from"), *GetDebugName(), SessionDirectory.ToString());
+		return false;
+	}
+
+	UE_LOG(LogVirtualization, Display, TEXT("[%s] Created directory '%s' to submit payloads from"), *GetDebugName(), SessionDirectory.ToString());
 
 	ON_SCOPE_EXIT
 	{
 		// Clean up the payload file from disk and the temp directories, but we do not need to give errors if any of these operations fail.
-		IFileManager::Get().DeleteDirectory(RootDirectory.ToString(), false, true);
+		IFileManager::Get().DeleteDirectory(SessionDirectory.ToString(), false, true);
 	};
 
 	TArray<FString> FilesToSubmit;
 	FilesToSubmit.Reserve(Requests.Num());
 
-	// Write the payloads to disk so that they can be submitted
+	// Write the payloads to disk so that they can be submitted (source control module currently requires the files to
+	// be on disk)
 	{
 		TRACE_CPUPROFILER_EVENT_SCOPE(FSourceControlBackend::PushData::CreateFiles);
 
@@ -247,11 +296,13 @@ bool FSourceControlBackend::PushData(TArrayView<FPushRequest> Requests)
 		{
 			TStringBuilder<52> LocalPayloadPath;
 			Utils::PayloadIdToPath(Request.Identifier, LocalPayloadPath);
-			FString PayloadFilePath = *WriteToString<512>(RootDirectory, LocalPayloadPath);
 
-			UE_LOG(LogVirtualization, Verbose, TEXT("[%s] Writing payload to '%s' for submission"), *GetDebugName(), *PayloadFilePath);
+			TStringBuilder<260> PayloadFilePath;
+			FPathViews::Append(PayloadFilePath, SessionDirectory, LocalPayloadPath);
 
-			TUniquePtr<FArchive> FileAr(IFileManager::Get().CreateFileWriter(*PayloadFilePath));
+			UE_LOG(LogVirtualization, Verbose, TEXT("[%s] Writing payload to '%s' for submission"), *GetDebugName(), PayloadFilePath.ToString());
+
+			TUniquePtr<FArchive> FileAr(IFileManager::Get().CreateFileWriter(PayloadFilePath.ToString()));
 			if (!FileAr)
 			{
 				TStringBuilder<MAX_SPRINTF> SystemErrorMsg;
@@ -260,7 +311,7 @@ bool FSourceControlBackend::PushData(TArrayView<FPushRequest> Requests)
 				UE_LOG(LogVirtualization, Error, TEXT("[%s] Failed to write payload '%s' contents to '%s' due to system error: %s"),
 					*GetDebugName(),
 					*LexToString(Request.Identifier),
-					*PayloadFilePath,
+					PayloadFilePath.ToString(),
 					SystemErrorMsg.ToString());
 
 				return false;
@@ -297,7 +348,7 @@ bool FSourceControlBackend::PushData(TArrayView<FPushRequest> Requests)
 	// Create a temp workspace so that we can submit the payload from
 	{
 		TRACE_CPUPROFILER_EVENT_SCOPE(FSourceControlBackend::PushData::CreateWorkspace);
-		TSharedRef<FCreateWorkspace> CreateWorkspaceCommand = ISourceControlOperation::Create<FCreateWorkspace>(WorkspaceName, RootDirectory);
+		TSharedRef<FCreateWorkspace> CreateWorkspaceCommand = ISourceControlOperation::Create<FCreateWorkspace>(WorkspaceName, SessionDirectory);
 
 		TStringBuilder<512> DepotMapping;
 		DepotMapping << DepotRoot << TEXT("...");
@@ -494,6 +545,55 @@ void FSourceControlBackend::CreateDepotPath(const FIoHash& PayloadId, FStringBui
 	OutPath << DepotRoot << PayloadPath;
 }
 
+bool FSourceControlBackend::FindSubmissionWorkingDir(const FString& ConfigEntry)
+{
+	// Note regarding path lengths.
+	// During submission each payload path will be 90 characters in length which will then be appended to
+	// the SubmissionWorkingDir
+
+	SubmissionRootDir = FPlatformMisc::GetEnvironmentVariable(TEXT("UE-VirtualizationWorkingDir"));
+
+	if (!SubmissionRootDir.IsEmpty())
+	{
+		FPaths::NormalizeDirectoryName(SubmissionRootDir);
+		UE_LOG(LogVirtualization, Display, TEXT("[%s] Found Environment Variable: UE-VirtualizationWorkingDir"), *GetDebugName());	
+	}
+	else
+	{
+
+		bool bSubmitFromTempDir = false;
+		FParse::Bool(*ConfigEntry, TEXT("SubmitFromTempDir="), bSubmitFromTempDir);
+
+		TStringBuilder<260> PathBuilder;
+		if (bSubmitFromTempDir)
+		{
+			FPathViews::Append(PathBuilder, FPlatformProcess::UserTempDir(), TEXT("UnrealEngine/VASubmission"));	
+		}
+		else
+		{
+			FPathViews::Append(PathBuilder, FPaths::ProjectSavedDir(), TEXT("VASubmission"));
+		}
+
+		SubmissionRootDir = PathBuilder;
+	}
+
+	if (IFileManager::Get().DirectoryExists(*SubmissionRootDir) || IFileManager::Get().MakeDirectory(*SubmissionRootDir))
+	{
+		UE_LOG(LogVirtualization, Display, TEXT("[%s] Setting '%s' as the working directory"), *GetDebugName(), *SubmissionRootDir);
+		return true;
+	}
+	else
+	{
+		TStringBuilder<MAX_SPRINTF> SystemErrorMsg;
+		Utils::GetFormattedSystemError(SystemErrorMsg);
+
+		UE_LOG(LogVirtualization, Error, TEXT("[%s] Failed to set the  working directory to '%s' due to %s"), *GetDebugName(), *SubmissionRootDir, SystemErrorMsg.ToString());
+		SubmissionRootDir.Empty();
+
+		return false;
+	}
+}
+
 void FSourceControlBackend::OnConnectionError()
 {
 	auto Callback = [](float Delta)->bool
diff --git a/Engine/Source/Developer/Virtualization/Private/VirtualizationSourceControlBackend.h b/Engine/Source/Developer/Virtualization/Private/VirtualizationSourceControlBackend.h
--- a/Engine/Source/Developer/Virtualization/Private/VirtualizationSourceControlBackend.h
+++ b/Engine/Source/Developer/Virtualization/Private/VirtualizationSourceControlBackend.h
@@ -17,15 +17,24 @@ namespace UE::Virtualization
  * backend, found in Utils::PayloadIdToPath.
  *
  * Ini file setup:
- * 'Name'=(Type=SourceControl, DepotRoot="//XXX/")
+ * 'Name'=(Type=SourceControl, DepotRoot="//XXX/", UsePartitionedClient=X, SubmitFromTempDir=X)
  * Where 'Name' is the backend name in the hierarchy and 'XXX' is the path in the source control
  * depot where the payload files are being stored.
  * 
  * Optional Values:
- * UsePartitionedClient:	When true the temporary workspace client created to submit payloads 
- *							from will be created as a partitioned workspace which is less overhead
- *							on the source control server. If your server does not support this then
-  *							use false. [Default=True]
+ * UsePartitionedClient [bool]:	When true the temporary workspace client created to submit payloads 
+ *								from will be created as a partitioned workspace which is less overhead
+ *								on the source control server. If your server does not support this then
+ *								use false. [Default=True]
+ * SubmitFromTempDir [bool]:	When set to true, payloads will be submitted from the temp directory of
+ *								the current machine and when false the files will be submitted from the
+ *								Save directory of the current project. [Default=false]
+ * 
+ * Environment Variables:
+ * UE-VirtualizationWorkingDir [string]:	This can be set to a valid directory path that the backend
+ *											should use as the root location to submit payloads from.
+ *											If the users machine has this set then 'SubmitFromTempDir' 
+ *											will be ignored. 
  */
 class FSourceControlBackend final : public IVirtualizationBackend
 {
@@ -51,12 +60,17 @@ private:
 
 	void CreateDepotPath(const FIoHash& PayloadId, FStringBuilderBase& OutPath);
 
+	bool FindSubmissionWorkingDir(const FString& ConfigEntry);
+
 	/** Will display a FMessage notification to the user on the next valid engine tick to try and keep them aware of connection failures */
 	void OnConnectionError();
 	
 	/** The root where the virtualized payloads are stored in source control */
 	FString DepotRoot;
 
+	/** The root directory from which payloads are submitted. */
+	FString SubmissionRootDir;
+
 	/** Should we try to make the temp client partitioned or not? */
 	bool bUsePartitionedClient = true;
 
