Allow a specific config file to be passed to the virtualization system when it is being created to parse settings from.

#rnx

- IVirtualizationSystem now has a new ::Initialize method which is called after it is first created. This allows the system to return a proper error when setting up which was not possible when the set up code ran in the constructor.
- The method also takes a FConfigFile reference which should be used when parsing settings from the config file system.
-- The config file being passed to this method, is either the optional config file provided when calling UE::Virtualization::Initialize, or (if that was not set) we find the GEngineIni config file from the global cache to use.
- This should ensure that the entire virtualization system is now loading the settings from the same set of config files where as before the code was a bit inconsistent (loading it in place, using the GConfig global, or trying to find the file from the global cache)
- If the system fails to create we now log if the problem was that we could not find a valid factory for the given system (meaning the module the system lives in was not yet loaded) or if the initialization phase failed.
- FVirtualizationManager now passes a single FConfigFile into all of it's methods that are parsing settings from it.

diff --git a/Engine/Source/Developer/Virtualization/Private/VirtualizationManager.cpp b/Engine/Source/Developer/Virtualization/Private/VirtualizationManager.cpp
--- a/Engine/Source/Developer/Virtualization/Private/VirtualizationManager.cpp
+++ b/Engine/Source/Developer/Virtualization/Private/VirtualizationManager.cpp
@@ -268,22 +268,6 @@ FVirtualizationManager::FVirtualizationManager()
 	// TODO: We should just be able to call the logging in the destructor, but 
 	// we need to fix the startup/shutdown ordering of Mirage first.
 	COOK_STAT(FCoreDelegates::OnExit.AddStatic(Profiling::LogStats));
-
-	FConfigFile* ConfigFile = GConfig->Find(GEngineIni);
-	if (ConfigFile != nullptr)
-	{
-		ApplySettingsFromConfigFiles(*ConfigFile);
-		ApplyDebugSettingsFromConfigFiles(*ConfigFile);
-	}
-	else
-	{
-		UE_LOG(LogVirtualization, Error, TEXT("Failed to load config file settings for content virtualization"));
-	}
-
-	ApplySettingsFromCmdline();
-	ApplyDebugSettingsFromFromCmdline();
-	
-	MountBackends();
 }
 
 FVirtualizationManager::~FVirtualizationManager()
@@ -299,6 +283,21 @@ FVirtualizationManager::~FVirtualizationManager()
 	UE_LOG(LogVirtualization, Log, TEXT("Virtualization manager destroyed"));
 }
 
+bool FVirtualizationManager::Initialize(const FConfigFile& ConfigFile)
+{
+	TRACE_CPUPROFILER_EVENT_SCOPE(FVirtualizationManager::Initialize);
+
+	ApplySettingsFromConfigFiles(ConfigFile);
+	ApplyDebugSettingsFromConfigFiles(ConfigFile);
+	
+	ApplySettingsFromCmdline();
+	ApplyDebugSettingsFromFromCmdline();
+
+	MountBackends(ConfigFile);
+
+	return true;
+}
+
 bool FVirtualizationManager::IsEnabled() const
 {
 	return !AllBackends.IsEmpty();
@@ -640,12 +639,12 @@ void FVirtualizationManager::GetPayloadActivityInfo( GetPayloadActivityInfoFuncR
 	
 }
 
-void FVirtualizationManager::ApplySettingsFromConfigFiles(const FConfigFile& PlatformEngineIni)
+void FVirtualizationManager::ApplySettingsFromConfigFiles(const FConfigFile& ConfigFile)
 {
 	UE_LOG(LogVirtualization, Display, TEXT("Loading virtualization manager settings from config files..."));
 	
 	bool bEnablePayloadPushingFromIni = false;
-	if (PlatformEngineIni.GetBool(TEXT("Core.ContentVirtualization"), TEXT("EnablePushToBackend"), bEnablePayloadPushingFromIni))
+	if (ConfigFile.GetBool(TEXT("Core.ContentVirtualization"), TEXT("EnablePushToBackend"), bEnablePayloadPushingFromIni))
 	{
 		bEnablePayloadPushing = bEnablePayloadPushingFromIni;
 		UE_LOG(LogVirtualization, Display, TEXT("\tEnablePushToBackend : %s"), bEnablePayloadPushing ? TEXT("true") : TEXT("false") );
@@ -656,7 +655,7 @@ void FVirtualizationManager::ApplySettingsFromConfigFiles(const FConfigFile& Pla
 	}
 
 	bool bEnableCacheAfterPullFromIni = false;
-	if (PlatformEngineIni.GetBool(TEXT("Core.ContentVirtualization"), TEXT("EnableCacheAfterPull"), bEnableCacheAfterPullFromIni))
+	if (ConfigFile.GetBool(TEXT("Core.ContentVirtualization"), TEXT("EnableCacheAfterPull"), bEnableCacheAfterPullFromIni))
 	{
 		bEnableCacheAfterPull = bEnableCacheAfterPullFromIni;
 		UE_LOG(LogVirtualization, Display, TEXT("\tCachePulledPayloads : %s"), bEnableCacheAfterPull ? TEXT("true") : TEXT("false"));
@@ -667,7 +666,7 @@ void FVirtualizationManager::ApplySettingsFromConfigFiles(const FConfigFile& Pla
 	}
 
 	int64 MinPayloadLengthFromIni = 0;
-	if (PlatformEngineIni.GetInt64(TEXT("Core.ContentVirtualization"), TEXT("MinPayloadLength"), MinPayloadLengthFromIni))
+	if (ConfigFile.GetInt64(TEXT("Core.ContentVirtualization"), TEXT("MinPayloadLength"), MinPayloadLengthFromIni))
 	{
 		MinPayloadLength = MinPayloadLengthFromIni;
 		UE_LOG(LogVirtualization, Display, TEXT("\tMinPayloadLength : %" INT64_FMT), MinPayloadLength );
@@ -678,7 +677,7 @@ void FVirtualizationManager::ApplySettingsFromConfigFiles(const FConfigFile& Pla
 	}
 
 	FString BackendGraphNameFromIni;
-	if (PlatformEngineIni.GetString(TEXT("Core.ContentVirtualization"), TEXT("BackendGraph"), BackendGraphNameFromIni))
+	if (ConfigFile.GetString(TEXT("Core.ContentVirtualization"), TEXT("BackendGraph"), BackendGraphNameFromIni))
 	{
 		BackendGraphName = BackendGraphNameFromIni;
 		UE_LOG(LogVirtualization, Display, TEXT("\tBackendGraphName : %s"), *BackendGraphName );
@@ -689,7 +688,7 @@ void FVirtualizationManager::ApplySettingsFromConfigFiles(const FConfigFile& Pla
 	}
 
 	FString FilterModeFromIni;
-	if (PlatformEngineIni.GetString(TEXT("Core.ContentVirtualization"), TEXT("FilterMode"), FilterModeFromIni))
+	if (ConfigFile.GetString(TEXT("Core.ContentVirtualization"), TEXT("FilterMode"), FilterModeFromIni))
 	{
 		if(LexTryParseString(FilteringMode, FilterModeFromIni))
 		{
@@ -706,7 +705,7 @@ void FVirtualizationManager::ApplySettingsFromConfigFiles(const FConfigFile& Pla
 	}
 
 	bool bFilterEngineContentFromIni = true;
-	if (PlatformEngineIni.GetBool(TEXT("Core.ContentVirtualization"), TEXT("FilterEngineContent"), bFilterEngineContentFromIni))
+	if (ConfigFile.GetBool(TEXT("Core.ContentVirtualization"), TEXT("FilterEngineContent"), bFilterEngineContentFromIni))
 	{
 		bFilterEngineContent = bFilterEngineContentFromIni;
 		UE_LOG(LogVirtualization, Display, TEXT("\tFilterEngineContent : %s"), bFilterEngineContent ? TEXT("true") : TEXT("false"));
@@ -717,7 +716,7 @@ void FVirtualizationManager::ApplySettingsFromConfigFiles(const FConfigFile& Pla
 	}
 	
 	bool bFilterEnginePluginContentFromIni = true;
-	if (PlatformEngineIni.GetBool(TEXT("Core.ContentVirtualization"), TEXT("FilterEnginePluginContent"), bFilterEnginePluginContentFromIni))
+	if (ConfigFile.GetBool(TEXT("Core.ContentVirtualization"), TEXT("FilterEnginePluginContent"), bFilterEnginePluginContentFromIni))
 	{
 		bFilterEnginePluginContent = bFilterEnginePluginContentFromIni;
 		UE_LOG(LogVirtualization, Display, TEXT("\tFilterEnginePluginContent : %s"), bFilterEnginePluginContent ? TEXT("true") : TEXT("false"));
@@ -744,20 +743,20 @@ void FVirtualizationManager::ApplySettingsFromCmdline()
 	}
 }
 
-void FVirtualizationManager::ApplyDebugSettingsFromConfigFiles(const FConfigFile& PlatformEngineIni)
+void FVirtualizationManager::ApplyDebugSettingsFromConfigFiles(const FConfigFile& ConfigFile)
 {
 	UE_LOG(LogVirtualization, Display, TEXT("Loading virtualization manager debugging settings from config files..."));
 
 	// Note that the debug settings are optional and could be left out of the config files entirely
 	bool bForceSingleThreadedFromIni = false;
-	if (PlatformEngineIni.GetBool(TEXT("Core.ContentVirtualizationDebugOptions"), TEXT("ForceSingleThreaded"), bForceSingleThreadedFromIni))
+	if (ConfigFile.GetBool(TEXT("Core.ContentVirtualizationDebugOptions"), TEXT("ForceSingleThreaded"), bForceSingleThreadedFromIni))
 	{
 		bForceSingleThreaded = bForceSingleThreadedFromIni;
 		UE_LOG(LogVirtualization, Display, TEXT("\tForceSingleThreaded : %s"), bForceSingleThreaded ? TEXT("true") : TEXT("false"));
 	}
 
 	bool bValidateAfterPushOperationFromIni = false;
-	if (PlatformEngineIni.GetBool(TEXT("Core.ContentVirtualizationDebugOptions"), TEXT("ValidateAfterPushOperation"), bValidateAfterPushOperationFromIni))
+	if (ConfigFile.GetBool(TEXT("Core.ContentVirtualizationDebugOptions"), TEXT("ValidateAfterPushOperation"), bValidateAfterPushOperationFromIni))
 	{
 		bValidateAfterPushOperation = bValidateAfterPushOperationFromIni;
 		UE_LOG(LogVirtualization, Display, TEXT("\tValidateAfterPushOperation : %s"), bValidateAfterPushOperation ? TEXT("true") : TEXT("false"));
@@ -784,7 +783,7 @@ void FVirtualizationManager::ApplyDebugSettingsFromFromCmdline()
 	}
 }
 
-void FVirtualizationManager::MountBackends()
+void FVirtualizationManager::MountBackends(const FConfigFile& ConfigFile)
 {
 	TRACE_CPUPROFILER_EVENT_SCOPE(FVirtualizationManager::MountBackends);
 
@@ -793,7 +792,7 @@ void FVirtualizationManager::MountBackends()
 
 	const TCHAR* GraphName = *BackendGraphName;
 
-	if(!GConfig->DoesSectionExist(GraphName, GEngineIni))
+	if(!ConfigFile.DoesSectionExist(GraphName))
 	{
 		UE_LOG(LogVirtualization, Fatal, TEXT("Unable to find the backend graph: '%s' [ini=%s]."), GraphName, *GEngineIni);
 	}
@@ -802,14 +801,14 @@ void FVirtualizationManager::MountBackends()
 
 	// It is important to parse the local storage hierarchy first so those backends will show up before the
 	// persistent storage backends in 'PullEnabledBackends'.
-	ParseHierarchy(GraphName, TEXT("LocalStorageHierarchy"), FactoryLookupTable, LocalCachableBackends);
-	ParseHierarchy(GraphName, TEXT("PersistentStorageHierarchy"), FactoryLookupTable, PersistentStorageBackends);
+	ParseHierarchy(ConfigFile, GraphName, TEXT("LocalStorageHierarchy"), FactoryLookupTable, LocalCachableBackends);
+	ParseHierarchy(ConfigFile, GraphName, TEXT("PersistentStorageHierarchy"), FactoryLookupTable, PersistentStorageBackends);
 }
 
-void FVirtualizationManager::ParseHierarchy(const TCHAR* GraphName, const TCHAR* HierarchyKey, const FRegistedFactories& FactoryLookupTable, FBackendArray& PushArray)
+void FVirtualizationManager::ParseHierarchy(const FConfigFile& ConfigFile, const TCHAR* GraphName, const TCHAR* HierarchyKey, const FRegistedFactories& FactoryLookupTable, FBackendArray& PushArray)
 {
 	FString HierarchyData;
-	if (!GConfig->GetString(GraphName, HierarchyKey, HierarchyData, GEngineIni))
+	if (!ConfigFile.GetValue(GraphName, HierarchyKey, HierarchyData))
 	{
 		UE_LOG(LogVirtualization, Fatal, TEXT("Unable to find the '%s' entry for the content virtualization backend graph '%s' [ini=%s]."), HierarchyKey, GraphName, *GEngineIni);
 	}
@@ -825,18 +824,18 @@ void FVirtualizationManager::ParseHierarchy(const TCHAR* GraphName, const TCHAR*
 
 	for (const FString& Entry : Entries)
 	{
-		CreateBackend(GraphName, Entry, FactoryLookupTable, PushArray);
+		CreateBackend(ConfigFile, GraphName, Entry, FactoryLookupTable, PushArray);
 	}
 }
 
-bool FVirtualizationManager::CreateBackend(const TCHAR* GraphName, const FString& ConfigEntryName, const FRegistedFactories& FactoryLookupTable, FBackendArray& PushArray)
+bool FVirtualizationManager::CreateBackend(const FConfigFile& ConfigFile, const TCHAR* GraphName, const FString& ConfigEntryName, const FRegistedFactories& FactoryLookupTable, FBackendArray& PushArray)
 {
 	// All failures in this method are considered fatal, however it still returns true/false in case we decide
 	// to be more forgiving in the future.
 	UE_LOG(LogVirtualization, Display, TEXT("Mounting backend entry '%s'"), *ConfigEntryName);
 
 	FString BackendData;
-	if (!GConfig->GetString(GraphName, *ConfigEntryName, BackendData, GEngineIni))
+	if (!ConfigFile.GetValue(GraphName, *ConfigEntryName, BackendData))
 	{
 		UE_LOG(LogVirtualization, Fatal, TEXT("Unable to find the entry '%s' in the content virtualization backend graph '%s' [ini=%s]."), *ConfigEntryName, GraphName, *GEngineIni);
 		return false;
diff --git a/Engine/Source/Developer/Virtualization/Private/VirtualizationManager.h b/Engine/Source/Developer/Virtualization/Private/VirtualizationManager.h
--- a/Engine/Source/Developer/Virtualization/Private/VirtualizationManager.h
+++ b/Engine/Source/Developer/Virtualization/Private/VirtualizationManager.h
@@ -105,6 +105,8 @@ public:
 private:
 	/* IVirtualizationSystem implementation */
 
+	virtual bool Initialize(const FConfigFile& ConfigFile) override;
+
 	virtual bool IsEnabled() const override;
 	virtual bool IsPushingEnabled(EStorageType StorageType) const override;
 	
@@ -126,15 +128,15 @@ private:
 	
 private:
 
-	void ApplySettingsFromConfigFiles(const FConfigFile& PlatformEngineIni);
+	void ApplySettingsFromConfigFiles(const FConfigFile& ConfigFile);
 	void ApplySettingsFromCmdline();
 	
-	void ApplyDebugSettingsFromConfigFiles(const FConfigFile& PlatformEngineIni);
+	void ApplyDebugSettingsFromConfigFiles(const FConfigFile& ConfigFile);
 	void ApplyDebugSettingsFromFromCmdline();
 
-	void MountBackends();
-	void ParseHierarchy(const TCHAR* GraphName, const TCHAR* HierarchyKey, const FRegistedFactories& FactoryLookupTable, FBackendArray& PushArray);
-	bool CreateBackend(const TCHAR* GraphName, const FString& ConfigEntryName, const FRegistedFactories& FactoryLookupTable, FBackendArray& PushArray);
+	void MountBackends(const FConfigFile& ConfigFile);
+	void ParseHierarchy(const FConfigFile& ConfigFile, const TCHAR* GraphName, const TCHAR* HierarchyKey, const FRegistedFactories& FactoryLookupTable, FBackendArray& PushArray);
+	bool CreateBackend(const FConfigFile& ConfigFile, const TCHAR* GraphName, const FString& ConfigEntryName, const FRegistedFactories& FactoryLookupTable, FBackendArray& PushArray);
 
 	void AddBackend(TUniquePtr<IVirtualizationBackend> Backend, FBackendArray& PushArray);
 
diff --git a/Engine/Source/Runtime/Core/Private/Virtualization/VirtualizationSystem.cpp b/Engine/Source/Runtime/Core/Private/Virtualization/VirtualizationSystem.cpp
--- a/Engine/Source/Runtime/Core/Private/Virtualization/VirtualizationSystem.cpp
+++ b/Engine/Source/Runtime/Core/Private/Virtualization/VirtualizationSystem.cpp
@@ -19,6 +19,11 @@ public:
 
 	virtual ~FNullVirtualizationSystem() = default;
 
+	virtual bool Initialize(const FConfigFile& ConfigFile) override
+	{
+		return true;
+	}
+
 	virtual bool IsEnabled() const override
 	{
 		return false;
@@ -70,36 +75,66 @@ public:
 
 TUniquePtr<IVirtualizationSystem> GVirtualizationSystem = nullptr;
 
-void Initialize()
+/** Utility function for finding a IVirtualizationSystemFactory for a given system name */
+Private::IVirtualizationSystemFactory* FindFactory(FName SystemName)
+{
+	TArray<Private::IVirtualizationSystemFactory*> AvaliableSystems = IModularFeatures::Get().GetModularFeatureImplementations<Private::IVirtualizationSystemFactory>(FName("VirtualizationSystem"));
+	for (Private::IVirtualizationSystemFactory* SystemFactory : AvaliableSystems)
+	{
+		if (SystemFactory->GetName() == SystemName)
+		{
+			return SystemFactory;
+		}
+	}
+
+	return nullptr;
+}
+
+void Initialize(FConfigFile* ConfigFile)
 {
 	FName SystemName;
 
-	FConfigFile EngineIni;
-	if (FConfigCacheIni::LoadLocalIniFile(EngineIni, TEXT("Engine"), true))
+	if (ConfigFile == nullptr)
+	{
+		ConfigFile = GConfig->Find(GEngineIni);
+	}
+	
+	if (ConfigFile != nullptr)
 	{
 		FString RawSystemName;
-		if (EngineIni.GetString(TEXT("Core.ContentVirtualization"), TEXT("SystemName"), RawSystemName))
+		if (ConfigFile->GetString(TEXT("Core.ContentVirtualization"), TEXT("SystemName"), RawSystemName))
 		{
 			SystemName = FName(RawSystemName);
 			UE_LOG(LogVirtualization, Display, TEXT("VirtualizationSystem name found in ini file: %s"), *RawSystemName);
 		}
 	}
-
+	
 	if (!SystemName.IsNone())
 	{
-		TArray<Private::IVirtualizationSystemFactory*> AvaliableSystems = IModularFeatures::Get().GetModularFeatureImplementations<Private::IVirtualizationSystemFactory>(FName("VirtualizationSystem"));
-		for (Private::IVirtualizationSystemFactory* SystemFactory : AvaliableSystems)
+		Private::IVirtualizationSystemFactory* SystemFactory = FindFactory(SystemName);
+		if (SystemFactory != nullptr)
 		{
-			if (SystemFactory->GetName() == SystemName)
+			GVirtualizationSystem = SystemFactory->Create();
+			check(GVirtualizationSystem.IsValid()); // It is assumed that create will always return a valid pointer
+
+			if (!GVirtualizationSystem->Initialize(*ConfigFile))
 			{
-				GVirtualizationSystem = SystemFactory->Create();
-				return;
+				UE_LOG(LogVirtualization, Error, TEXT("Initialization of the virtualization system '%s' failed, falling back to the default implementation"), *SystemName.ToString());
+				GVirtualizationSystem.Reset();
 			}
 		}
+		else
+		{
+			UE_LOG(LogVirtualization, Error, TEXT("Unable to find factory to create the virtualization system: %s"), *SystemName.ToString());
+		}
 	}
 
-	// We found no system to create so we will use the fallback Null system
-	GVirtualizationSystem = MakeUnique<FNullVirtualizationSystem>();	
+	if (!GVirtualizationSystem.IsValid())
+	{
+		// We found no system to create so we will use the fallback Null system
+		GVirtualizationSystem = MakeUnique<FNullVirtualizationSystem>();
+		GVirtualizationSystem->Initialize(*ConfigFile); // We know this method does nothing but we call it anyway for sake of form
+	}
 }
 
 void Shutdown()
@@ -114,7 +149,7 @@ IVirtualizationSystem& IVirtualizationSystem::Get()
 	if (GVirtualizationSystem == nullptr)
 	{
 		UE_LOG(LogVirtualization, Warning, TEXT("UE::Virtualization::Initialize was not called before UE::Virtualization::IVirtualizationSystem::Get()!"));
-		Initialize();
+		UE::Virtualization::Initialize();
 	}
 
 	return *GVirtualizationSystem;
diff --git a/Engine/Source/Runtime/Core/Public/Virtualization/VirtualizationSystem.h b/Engine/Source/Runtime/Core/Public/Virtualization/VirtualizationSystem.h
--- a/Engine/Source/Runtime/Core/Public/Virtualization/VirtualizationSystem.h
+++ b/Engine/Source/Runtime/Core/Public/Virtualization/VirtualizationSystem.h
@@ -117,7 +117,7 @@ struct FPushRequest
  * during process start up but it will also be called by IVirtualizationSystem::Get if it detects that the
  * IVirtualizationSystem has not yet been set up.
  */
-CORE_API void Initialize();
+CORE_API void Initialize(FConfigFile* ConfigFile = nullptr);
 
 /**
  * Shutdowns the global IVirtualizationSystem if it exists. 
@@ -145,6 +145,20 @@ public:
 	IVirtualizationSystem() = default;
 	virtual ~IVirtualizationSystem() = default;
 
+	/**
+	 * Initialize the system from the given config file. Since the system might be initialized from a custom set
+	 * of specifically loaded config files, this is the only point in the systems life time that you are guaranteed 
+	 * to have the correct config file available so make sure that all settings that are required are parsed and
+	 * stored at this point rather than trying to access the config system directly in the future.
+	 * 
+	 * @param ConfigFile	The config file to parse the systems settings from. The object is guaranteed to be
+	 *						valid for the life time of the method so do not store it for future use.
+	 * @return				True if the system was initialized correctly, otherwise false. Note that if the method
+	 *						returns false then the system will be deleted and the default FNullVirtualizationSystem 
+	 *						will be used instead.
+	 */
+	virtual bool Initialize(const FConfigFile& ConfigFile) = 0;
+
 	/** Gain access to the current virtualization system active for the project */
 	static IVirtualizationSystem& Get();
 
