The source control virtualization backend now creates and uses it's own source control provider

#rnx

- The backend can now connect to source control without forcing the rest of the editor to connect to source control.
- If the user disconnects the editor from source control, the backend will still be able to function using it's own connection.
-When the backend needs to change the current workspace in order to submit payloads we no longer invalidate the file cache used by the editor and so the user will no longer have to manually refresh files in the content browsers to restore the source control status icon.

diff --git a/Engine/Source/Developer/Virtualization/Private/VirtualizationSourceControlBackend.cpp b/Engine/Source/Developer/Virtualization/Private/VirtualizationSourceControlBackend.cpp
--- a/Engine/Source/Developer/Virtualization/Private/VirtualizationSourceControlBackend.cpp
+++ b/Engine/Source/Developer/Virtualization/Private/VirtualizationSourceControlBackend.cpp
@@ -13,6 +13,7 @@
 #include "Misc/PathViews.h"
 #include "Misc/Paths.h"
 #include "Misc/ScopeExit.h"
+#include "SourceControlInitSettings.h"
 #include "SourceControlOperations.h"
 #include "VirtualizationSourceControlUtilities.h"
 #include "VirtualizationUtilities.h"
@@ -106,32 +107,22 @@ bool FSourceControlBackend::Initialize(const FString& ConfigEntry)
 	FParse::Bool(*ConfigEntry, TEXT("UsePartitionedClient="), bUsePartitionedClient);
 	UE_LOG(LogVirtualization, Display, TEXT("[%s] Using partitioned clients: '%s'"), *GetDebugName(), bUsePartitionedClient ? TEXT("true") : TEXT("false"));
 
-	ISourceControlModule& SSCModule = ISourceControlModule::Get();
+	// We do not want the connection to have a client workspace so explicitly set it to empty
+	FSourceControlInitSettings SCCSettings(FSourceControlInitSettings::EBehavior::OverrideExisting);
+	SCCSettings.AddSetting(TEXT("P4Client"), TEXT(""));
 
-	// We require perforce as the source control provider as it is currently the only one that has the virtualization functionality implemented
-	const FName SourceControlName = SSCModule.GetProvider().GetName();
-	if (SourceControlName.IsNone())
+	SCCProvider = ISourceControlModule::Get().CreateProvider(FName("Perforce"), TEXT("Virtualization"), SCCSettings);
+	if (!SCCProvider.IsValid())
 	{
-		// No source control provider is set so we can try to set it to "Perforce"
-		// Note this call will fatal error if "Perforce" is not a valid option
-		SSCModule.SetProvider(FName("Perforce"));
-	}
-	else if (SourceControlName != TEXT("Perforce"))
-	{
-		UE_LOG(LogVirtualization, Error, TEXT("Attempting to initialize FSourceControlBackend but source control is '%s' and only Perforce is currently supported!"), *SourceControlName.ToString());
+		UE_LOG(LogVirtualization, Error, TEXT("[%s] Failed to create a perforce connection, this seems to be unsupported by the editor"), *GetDebugName());
 		return false;
 	}
 
-	ISourceControlProvider& SCCProvider = SSCModule.GetProvider();
-
-	if (!SCCProvider.IsAvailable())
-	{
-		SCCProvider.Init();
-	}
-
+	SCCProvider->Init(true);
+ 
 	// Note that if the connect is failing then we expect it to fail here rather than in the subsequent attempts to get the meta info file
 	TSharedRef<FConnect, ESPMode::ThreadSafe> ConnectCommand = ISourceControlOperation::Create<FConnect>();
-	if (SCCProvider.Execute(ConnectCommand, FString(), EConcurrency::Synchronous) != ECommandResult::Succeeded)
+	if (SCCProvider->Execute(ConnectCommand, FString(), EConcurrency::Synchronous) != ECommandResult::Succeeded)
 	{		
 		FTextBuilder Errors;
 		for (const FText& Msg : ConnectCommand->GetResultInfo().ErrorMessages)
@@ -140,7 +131,7 @@ bool FSourceControlBackend::Initialize(const FString& ConfigEntry)
 		}
 
 		FMessageLog Log("LogVirtualization");
-		Log.Error(	FText::Format(LOCTEXT("FailedSourceControlConnection", "Failed to connect to source control backend with the following errors:\n{0}\nThe source control backend will be unable to pull payloads!\nTrying logging in with the 'p4 login' command or by using p4vs/UnrealGameSync."),
+		Log.Error(	FText::Format(LOCTEXT("FailedSourceControlConnection", "Failed to connect to source control backend with the following errors:\n{0}\nThe source control backend had trouble connecting!\nTrying logging in with the 'p4 login' command or by using p4vs/UnrealGameSync."),
 					Errors.ToText()));
 		
 		OnConnectionError();
@@ -153,7 +144,7 @@ bool FSourceControlBackend::Initialize(const FString& ConfigEntry)
 
 #if IS_SOURCE_CONTROL_THREAD_SAFE
 	TSharedRef<FDownloadFile, ESPMode::ThreadSafe> DownloadCommand = ISourceControlOperation::Create<FDownloadFile>();
-	if (SCCProvider.Execute(DownloadCommand, PayloadMetaInfoPath, EConcurrency::Synchronous) != ECommandResult::Succeeded)
+	if (SCCProvider->Execute(DownloadCommand, PayloadMetaInfoPath, EConcurrency::Synchronous) != ECommandResult::Succeeded)
 	{
 		FMessageLog Log("LogVirtualization");
 
@@ -165,7 +156,7 @@ bool FSourceControlBackend::Initialize(const FString& ConfigEntry)
 	}	
 #else
 	TSharedRef<FDownloadFile, ESPMode::ThreadSafe> DownloadCommand = ISourceControlOperation::Create<FDownloadFile>();
-	if (!SCCProvider.TryToDownloadFileFromBackgroundThread(DownloadCommand, PayloadMetaInfoPath))
+	if (!SCCProvider->TryToDownloadFileFromBackgroundThread(DownloadCommand, PayloadMetaInfoPath))
 	{
 		FMessageLog Log("LogVirtualization");
 
@@ -206,17 +197,15 @@ FCompressedBuffer FSourceControlBackend::PullData(const FIoHash& Id)
 
 	UE_LOG(LogVirtualization, Verbose, TEXT("[%s] Attempting to pull '%s' from source control"), *GetDebugName(), *DepotPath);
 
-	ISourceControlProvider& SCCProvider = ISourceControlModule::Get().GetProvider();
-
 #if IS_SOURCE_CONTROL_THREAD_SAFE
 	TSharedRef<FDownloadFile, ESPMode::ThreadSafe> DownloadCommand = ISourceControlOperation::Create<FDownloadFile>(FDownloadFile::EVerbosity::None);
-	if (SCCProvider.Execute(DownloadCommand, DepotPath.ToString(), EConcurrency::Synchronous) != ECommandResult::Succeeded)
+	if (SCCProvider->Execute(DownloadCommand, DepotPath.ToString(), EConcurrency::Synchronous) != ECommandResult::Succeeded)
 	{
 		return FCompressedBuffer();
 	}
 #else
 	TSharedRef<FDownloadFile> DownloadCommand = ISourceControlOperation::Create<FDownloadFile>(FDownloadFile::EVerbosity::None);
-	if (!SCCProvider.TryToDownloadFileFromBackgroundThread(DownloadCommand, DepotPath.ToString()))
+	if (!SCCProvider->TryToDownloadFileFromBackgroundThread(DownloadCommand, DepotPath.ToString()))
 	{
 		return FCompressedBuffer();
 	}
@@ -343,8 +332,6 @@ bool FSourceControlBackend::PushData(TArrayView<FPushRequest> Requests)
 	TStringBuilder<64> WorkspaceName;
 	WorkspaceName << TEXT("VASubmission-") << SessionGuid;
 
-	ISourceControlProvider& SCCProvider = ISourceControlModule::Get().GetProvider();
-
 	// Create a temp workspace so that we can submit the payload from
 	{
 		TRACE_CPUPROFILER_EVENT_SCOPE(FSourceControlBackend::PushData::CreateWorkspace);
@@ -365,7 +352,7 @@ bool FSourceControlBackend::PushData(TArrayView<FPushRequest> Requests)
 
 		CreateWorkspaceCommand->SetDescription(TEXT("This workspace was autogenerated when submitting virtualized payloads to source control"));
 
-		if (SCCProvider.Execute(CreateWorkspaceCommand) != ECommandResult::Succeeded)
+		if (SCCProvider->Execute(CreateWorkspaceCommand) != ECommandResult::Succeeded)
 		{
 			UE_LOG(LogVirtualization, Error, TEXT("[%s] Failed to create temp workspace '%s' to submit payloads from"),
 				*GetDebugName(),
@@ -378,7 +365,7 @@ bool FSourceControlBackend::PushData(TArrayView<FPushRequest> Requests)
 	ON_SCOPE_EXIT
 	{
 		// Remove the temp workspace mapping
-		if (SCCProvider.Execute(ISourceControlOperation::Create<FDeleteWorkspace>(WorkspaceName)) != ECommandResult::Succeeded)
+		if (SCCProvider->Execute(ISourceControlOperation::Create<FDeleteWorkspace>(WorkspaceName)) != ECommandResult::Succeeded)
 		{
 			UE_LOG(LogVirtualization, Warning, TEXT("[%s] Failed to remove temp workspace '%s' please delete manually"), *GetDebugName(), WorkspaceName.ToString());
 		}
@@ -389,7 +376,7 @@ bool FSourceControlBackend::PushData(TArrayView<FPushRequest> Requests)
 		TRACE_CPUPROFILER_EVENT_SCOPE(FSourceControlBackend::PushData::SwitchWorkspace);
 
 		FSourceControlResultInfo SwitchToNewWorkspaceInfo;
-		if (SCCProvider.SwitchWorkspace(WorkspaceName, SwitchToNewWorkspaceInfo, &OriginalWorkspace) != ECommandResult::Succeeded)
+		if (SCCProvider->SwitchWorkspace(WorkspaceName, SwitchToNewWorkspaceInfo, &OriginalWorkspace) != ECommandResult::Succeeded)
 		{
 			UE_LOG(LogVirtualization, Error, TEXT("[%s] Failed to switch to temp workspace '%s' when trying to submit payloads"),
 				*GetDebugName(),
@@ -402,7 +389,7 @@ bool FSourceControlBackend::PushData(TArrayView<FPushRequest> Requests)
 	ON_SCOPE_EXIT
 	{
 		FSourceControlResultInfo SwitchToOldWorkspaceInfo;
-		if (SCCProvider.SwitchWorkspace(OriginalWorkspace, SwitchToOldWorkspaceInfo, nullptr) != ECommandResult::Succeeded)
+		if (SCCProvider->SwitchWorkspace(OriginalWorkspace, SwitchToOldWorkspaceInfo, nullptr) != ECommandResult::Succeeded)
 		{
 			// Failing to restore the old workspace could result in confusing editor issues and data loss, so for now it is fatal.
 			// The medium term plan should be to refactor the SourceControlModule so that we could use an entirely different 
@@ -416,7 +403,7 @@ bool FSourceControlBackend::PushData(TArrayView<FPushRequest> Requests)
 	TArray<FSourceControlStateRef> FileStates;
 	{
 		TRACE_CPUPROFILER_EVENT_SCOPE(FSourceControlBackend::PushData::GetFileStates);
-		if (SCCProvider.GetState(FilesToSubmit, FileStates, EStateCacheUsage::ForceUpdate) != ECommandResult::Succeeded)
+		if (SCCProvider->GetState(FilesToSubmit, FileStates, EStateCacheUsage::ForceUpdate) != ECommandResult::Succeeded)
 		{
 			UE_LOG(LogVirtualization, Error, TEXT("[%s] Failed to find the current file state for payloads"), *GetDebugName());
 			return false;
@@ -460,7 +447,7 @@ bool FSourceControlBackend::PushData(TArrayView<FPushRequest> Requests)
 	{
 		TRACE_CPUPROFILER_EVENT_SCOPE(FSourceControlBackend::PushData::AddFiles);
 
-		if (SCCProvider.Execute(ISourceControlOperation::Create<FMarkForAdd>(), FilesToAdd) != ECommandResult::Succeeded)
+		if (SCCProvider->Execute(ISourceControlOperation::Create<FMarkForAdd>(), FilesToAdd) != ECommandResult::Succeeded)
 		{
 			UE_LOG(LogVirtualization, Error, TEXT("[%s] Failed to mark the payload file for Add in source control"), *GetDebugName());
 			return false;
@@ -478,7 +465,7 @@ bool FSourceControlBackend::PushData(TArrayView<FPushRequest> Requests)
 
 		CheckInOperation->SetDescription(FText::FromString(Description.ToString()));
 
-		if (SCCProvider.Execute(CheckInOperation, FilesToAdd) != ECommandResult::Succeeded)
+		if (SCCProvider->Execute(CheckInOperation, FilesToAdd) != ECommandResult::Succeeded)
 		{
 			UE_LOG(LogVirtualization, Error, TEXT("[%s] Failed to submit the payload file(s) to source control"), *GetDebugName());
 			return false;
@@ -496,8 +483,6 @@ bool FSourceControlBackend::PushData(TArrayView<FPushRequest> Requests)
 
 bool FSourceControlBackend::DoPayloadsExist(TArrayView<const FIoHash> PayloadIds, TArray<bool>& OutResults)
 {
-	ISourceControlProvider& SCCProvider = ISourceControlModule::Get().GetProvider();
-
 	TArray<FString> DepotPaths;
 	DepotPaths.Reserve(PayloadIds.Num());
 
@@ -514,7 +499,7 @@ bool FSourceControlBackend::DoPayloadsExist(TArrayView<const FIoHash> PayloadIds
 		}
 	}
 
-	ECommandResult::Type Result = SCCProvider.GetState(DepotPaths, PathStates, EStateCacheUsage::ForceUpdate);
+	ECommandResult::Type Result = SCCProvider->GetState(DepotPaths, PathStates, EStateCacheUsage::ForceUpdate);
 	if (Result != ECommandResult::Type::Succeeded)
 	{
 		UE_LOG(LogVirtualization, Error, TEXT("[%s] Failed to query the state of files in the source control depot"), *GetDebugName());
diff --git a/Engine/Source/Developer/Virtualization/Private/VirtualizationSourceControlBackend.h b/Engine/Source/Developer/Virtualization/Private/VirtualizationSourceControlBackend.h
--- a/Engine/Source/Developer/Virtualization/Private/VirtualizationSourceControlBackend.h
+++ b/Engine/Source/Developer/Virtualization/Private/VirtualizationSourceControlBackend.h
@@ -67,6 +67,9 @@ private:
 	/** Will display a FMessage notification to the user on the next valid engine tick to try and keep them aware of connection failures */
 	void OnConnectionError();
 
+	/** A source control connection owned by the backend*/
+	TUniquePtr<ISourceControlProvider> SCCProvider;
+	
 	/** The name of the current project */
 	FString ProjectName;
 
