Reduce the time spent checking if virtualized payloads already exist in source control storage or not.

#rnx

- Use the SetRequireDirPathEndWithSeperator option with the FUpdateStatus operation as we are checking depot file paths. If we do not use this option they can be seen as network paths and will be checked against the file system which can take some time to time out.
- Wrapped this up in a utility function so that it can be more easily updated in the future.

diff --git a/Engine/Source/Developer/Virtualization/Private/VirtualizationSourceControlBackend.cpp b/Engine/Source/Developer/Virtualization/Private/VirtualizationSourceControlBackend.cpp
--- a/Engine/Source/Developer/Virtualization/Private/VirtualizationSourceControlBackend.cpp
+++ b/Engine/Source/Developer/Virtualization/Private/VirtualizationSourceControlBackend.cpp
@@ -81,6 +81,20 @@ static void CreateDescription(const FString& ProjectName, const TArray<const FPu
 	}
 }
 
+[[nodiscard]] static ECommandResult::Type GetDepotPathStates(ISourceControlProvider& SCCProvider, const TArray<FString>& DepotPaths, TArray<FSourceControlStateRef>& OutStates)
+{
+	TSharedRef<FUpdateStatus, ESPMode::ThreadSafe> UpdateOperation = ISourceControlOperation::Create<FUpdateStatus>();
+	UpdateOperation->SetRequireDirPathEndWithSeparator(true);
+
+	ECommandResult::Type Result = SCCProvider.Execute(UpdateOperation, DepotPaths);
+	if (Result != ECommandResult::Succeeded)
+	{
+		return Result;
+	}
+
+	return SCCProvider.GetState(DepotPaths, OutStates, EStateCacheUsage::Use);
+}
+
 FSourceControlBackend::FSourceControlBackend(FStringView ProjectName, FStringView ConfigName, FStringView InDebugName)
 	: IVirtualizationBackend(ConfigName, InDebugName, EOperations::Both)
 	, ProjectName(ProjectName)
@@ -403,7 +417,7 @@ bool FSourceControlBackend::PushData(TArrayView<FPushRequest> Requests)
 	TArray<FSourceControlStateRef> FileStates;
 	{
 		TRACE_CPUPROFILER_EVENT_SCOPE(FSourceControlBackend::PushData::GetFileStates);
-		if (SCCProvider->GetState(FilesToSubmit, FileStates, EStateCacheUsage::ForceUpdate) != ECommandResult::Succeeded)
+		if (GetDepotPathStates(*SCCProvider, FilesToSubmit, FileStates) != ECommandResult::Succeeded)
 		{
 			UE_LOG(LogVirtualization, Error, TEXT("[%s] Failed to find the current file state for payloads"), *GetDebugName());
 			return false;
@@ -499,7 +513,7 @@ bool FSourceControlBackend::DoPayloadsExist(TArrayView<const FIoHash> PayloadIds
 		}
 	}
 
-	ECommandResult::Type Result = SCCProvider->GetState(DepotPaths, PathStates, EStateCacheUsage::ForceUpdate);
+	ECommandResult::Type Result = GetDepotPathStates(*SCCProvider, DepotPaths, PathStates);
 	if (Result != ECommandResult::Type::Succeeded)
 	{
 		UE_LOG(LogVirtualization, Error, TEXT("[%s] Failed to query the state of files in the source control depot"), *GetDebugName());
