Expose the package submission checks we run when submitting changelists from within the editor so that other systems can specifically invoke it in addition to running it when ever a source control module callback is invoked.

#rnx

- We need to expose the code that checks a set of packages for non virtualized payloads and virtualizes them so that the stand alone submission tool can access it.
- Currently I am adding this to the IVirtualizationSystem interface to avoid adding additional interface, but we might want to consider moving this some day as it is not a great fit with the rest of the interface.

diff --git a/Engine/Source/Developer/Virtualization/Private/PackageSubmissionChecks.cpp b/Engine/Source/Developer/Virtualization/Private/PackageSubmissionChecks.cpp
--- a/Engine/Source/Developer/Virtualization/Private/PackageSubmissionChecks.cpp
+++ b/Engine/Source/Developer/Virtualization/Private/PackageSubmissionChecks.cpp
@@ -124,9 +124,9 @@ bool TryCopyPackageWithoutTrailer(const FPackagePath PackagePath, const FString&
 	return true;
 }
 
-void OnPrePackageSubmission(const TArray<FString>& FilesToSubmit, TArray<FText>& DescriptionTags, TArray<FText>& Errors)
+void VirtualizePackages(const TArray<FString>& FilesToSubmit, TArray<FText>& OutDescriptionTags, TArray<FText>& OutErrors)
 {
-	TRACE_CPUPROFILER_EVENT_SCOPE(UE::Virtualization::OnPrePackageSubmission);
+	TRACE_CPUPROFILER_EVENT_SCOPE(UE::Virtualization::VirtualizePackages);
 
 	IVirtualizationSystem& System = IVirtualizationSystem::Get();
 
@@ -155,7 +155,7 @@ void OnPrePackageSubmission(const TArray<FString>& FilesToSubmit, TArray<FText>&
 	Progress.MakeDialog();
 
 	// Other systems may have added errors to this array, we need to check so later we can determine if this function added any additional errors.
-	const int32 NumErrors = Errors.Num();
+	const int32 NumErrors = OutErrors.Num();
 
 	struct FPackageInfo
 	{
@@ -197,7 +197,7 @@ void OnPrePackageSubmission(const TArray<FString>& FilesToSubmit, TArray<FText>&
 				{
 					FText Message = FText::Format(	LOCTEXT("Virtualization_PkgHasReferences", "Cannot virtualize the package '{1}' as it has referenced payloads in the trailer"),
 													FText::FromString(PackagePath.GetDebugName()));
-					Errors.Add(Message);
+					OutErrors.Add(Message);
 					return;
 				}
 
@@ -228,7 +228,7 @@ void OnPrePackageSubmission(const TArray<FString>& FilesToSubmit, TArray<FText>&
 	if (System.QueryPayloadStatuses(AllLocalPayloads, EStorageType::Persistent, PayloadStatuses) != EQueryResult::Success)
 	{
 		FText Message = LOCTEXT("Virtualization_DoesExistFail", "Failed to find the status of the payloads in the packages being submitted");
-		Errors.Add(Message);
+		OutErrors.Add(Message);
 
 		return;
 	}
@@ -253,7 +253,7 @@ void OnPrePackageSubmission(const TArray<FString>& FilesToSubmit, TArray<FText>&
 					FText Message = FText::Format(	LOCTEXT("Virtualization_UpdateStatusFailed", "Unable to update the status for the payload '{0}' in the package '{1}'"),
 													FText::FromString(LexToString(PackageInfo.LocalPayloads[Index])),
 													FText::FromString(PackageInfo.Path.GetDebugName()));
-					Errors.Add(Message);
+					OutErrors.Add(Message);
 					return;
 				}
 			}
@@ -292,7 +292,7 @@ void OnPrePackageSubmission(const TArray<FString>& FilesToSubmit, TArray<FText>&
 		{
 			FText Message = FText::Format(	LOCTEXT("Virtualization_PkgOpen", "Failed to open the package '{1}' for reading"),
 											FText::FromString(PackageInfo.Path.GetDebugName()));
-			Errors.Add(Message);
+			OutErrors.Add(Message);
 			return;
 		}
 
@@ -310,7 +310,7 @@ void OnPrePackageSubmission(const TArray<FString>& FilesToSubmit, TArray<FText>&
 												FText::FromString(PackageInfo.Path.GetDebugName()),
 												FText::FromString(LexToString(PayloadId)),
 												FText::FromString(LexToString(Payload.GetRawHash())));
-				Errors.Add(Message);
+				OutErrors.Add(Message);
 				return;
 			}
 
@@ -319,7 +319,7 @@ void OnPrePackageSubmission(const TArray<FString>& FilesToSubmit, TArray<FText>&
 				FText Message = FText::Format(	LOCTEXT("Virtualization_MissingPayload", "Unable to find the payload '{0}' in the local storage of package '{1}'"),
 												FText::FromString(LexToString(PayloadId)),
 												FText::FromString(PackageInfo.Path.GetDebugName()));
-				Errors.Add(Message);
+				OutErrors.Add(Message);
 				return;
 
 			}
@@ -333,7 +333,7 @@ void OnPrePackageSubmission(const TArray<FString>& FilesToSubmit, TArray<FText>&
 	if (!System.PushData(PayloadsToSubmit, EStorageType::Persistent))
 	{
 		FText Message = LOCTEXT("Virtualization_PushFailure", "Failed to push payloads");
-		Errors.Add(Message);
+		OutErrors.Add(Message);
 		return;
 	}
 
@@ -363,7 +363,7 @@ void OnPrePackageSubmission(const TArray<FString>& FilesToSubmit, TArray<FText>&
 					FText Message = FText::Format(	LOCTEXT("Virtualization_UpdateStatusFailed", "Unable to update the status for the payload '{0}' in the package '{1}'"),
 													FText::FromString(LexToString(Request.Identifier)),
 													FText::FromString(PackageInfo.Path.GetDebugName()));
-					Errors.Add(Message);
+					OutErrors.Add(Message);
 					return;
 				}
 			}
@@ -391,7 +391,7 @@ void OnPrePackageSubmission(const TArray<FString>& FilesToSubmit, TArray<FText>&
 		// TODO Optimization: Combine TryCopyPackageWithoutTrailer with the appending of the new trailer to avoid opening multiple handles
 
 		// Create copy of package minus the trailer the trailer
-		if (!TryCopyPackageWithoutTrailer(PackagePath, TempFilePath, PackageInfo.Trailer, Errors))
+		if (!TryCopyPackageWithoutTrailer(PackagePath, TempFilePath, PackageInfo.Trailer, OutErrors))
 		{
 			return;
 		}			
@@ -402,7 +402,7 @@ void OnPrePackageSubmission(const TArray<FString>& FilesToSubmit, TArray<FText>&
 		{
 			FText Message = FText::Format(	LOCTEXT("Virtualization_PkgOpen", "Failed to open the package '{1}' for reading"),
 											FText::FromString(PackagePath.GetDebugName()));
-			Errors.Add(Message);
+			OutErrors.Add(Message);
 			return;
 		}
 
@@ -411,7 +411,7 @@ void OnPrePackageSubmission(const TArray<FString>& FilesToSubmit, TArray<FText>&
 		{
 			FText Message = FText::Format(	LOCTEXT("Virtualization_TrailerAppendOpen", "Unable to open '{0}' to append the trailer'"),
 											FText::FromString(TempFilePath));
-			Errors.Add(Message);
+			OutErrors.Add(Message);
 			return;
 		}
 
@@ -420,7 +420,7 @@ void OnPrePackageSubmission(const TArray<FString>& FilesToSubmit, TArray<FText>&
 		{
 			FText Message = FText::Format(	LOCTEXT("Virtualization_TrailerAppend", "Failed to append the trailer to '{0}'"),
 											FText::FromString(TempFilePath));
-			Errors.Add(Message);
+			OutErrors.Add(Message);
 			return;
 		}
 
@@ -431,7 +431,7 @@ void OnPrePackageSubmission(const TArray<FString>& FilesToSubmit, TArray<FText>&
 
 	UE_LOG(LogVirtualization, Display, TEXT("%d package(s) had their trailer container modified and need to be updated"), PackagesToReplace.Num());
 
-	if (NumErrors == Errors.Num())
+	if (NumErrors == OutErrors.Num())
 	{
 		// TODO: Consider using the SavePackage model (move the original, then replace, so we can restore all of the original packages if needed)
 		// having said that, once a package is in PackagesToReplace it should still be safe to submit so maybe we don't need this level of protection?
@@ -456,7 +456,7 @@ void OnPrePackageSubmission(const TArray<FString>& FilesToSubmit, TArray<FText>&
 			{
 				FText Message = FText::Format(	LOCTEXT("Virtualization_MoveFailed", "Unable to replace the package '{0}' with the virtualized version"),
 												FText::FromString(Iterator.Key.GetDebugName()));
-				Errors.Add(Message);
+				OutErrors.Add(Message);
 				continue;
 			}
 		}
@@ -465,10 +465,10 @@ void OnPrePackageSubmission(const TArray<FString>& FilesToSubmit, TArray<FText>&
 	// If we had no new errors add the validation tag to indicate that the packages are safe for submission. 
 	// TODO: Currently this is a simple tag to make it easier for us to track which assets were submitted via the
 	// virtualization process in a test project. This should be expanded when we add proper p4 server triggers.
-	if (NumErrors == Errors.Num())
+	if (NumErrors == OutErrors.Num())
 	{
 		FText Tag = FText::FromString(TEXT("#virtualized"));
-		DescriptionTags.Add(Tag);
+		OutDescriptionTags.Add(Tag);
 	}
 
 	const double TimeInSeconds = FPlatformTime::Seconds() - StartTime;
diff --git a/Engine/Source/Developer/Virtualization/Private/PackageSubmissionChecks.h b/Engine/Source/Developer/Virtualization/Private/PackageSubmissionChecks.h
--- a/Engine/Source/Developer/Virtualization/Private/PackageSubmissionChecks.h
+++ b/Engine/Source/Developer/Virtualization/Private/PackageSubmissionChecks.h
@@ -9,6 +9,7 @@
 namespace UE::Virtualization
 {
 
-void OnPrePackageSubmission(const TArray<FString>& FilesToSubmit, TArray<FText>& DescriptionTags, TArray<FText>& Errors);
+/** @See IVirtualizationSystem::TryVirtualizePackages */
+void VirtualizePackages(const TArray<FString>& FilesToSubmit, TArray<FText>& OutDescriptionTags, TArray<FText>& OutErrors);
 
 } // namespace UE::Virtualization
diff --git a/Engine/Source/Developer/Virtualization/Private/Virtualization.cpp b/Engine/Source/Developer/Virtualization/Private/Virtualization.cpp
--- a/Engine/Source/Developer/Virtualization/Private/Virtualization.cpp
+++ b/Engine/Source/Developer/Virtualization/Private/Virtualization.cpp
@@ -26,7 +26,7 @@ public:
 		FDelayedAutoRegisterHelper(EDelayedRegisterRunPhase::EarliestPossiblePluginsLoaded, [this]() 
 			{
 				PackageSubmissionHandle = ISourceControlModule::Get().RegisterPreSubmitFinalize(
-					FSourceControlPreSubmitFinalizeDelegate::FDelegate::CreateStatic(&OnPrePackageSubmission));
+					FSourceControlPreSubmitFinalizeDelegate::FDelegate::CreateStatic(&VirtualizePackages));
 			});
 
 		FMessageLogModule& MessageLogModule = FModuleManager::LoadModuleChecked<FMessageLogModule>("MessageLog");
diff --git a/Engine/Source/Developer/Virtualization/Private/VirtualizationManager.cpp b/Engine/Source/Developer/Virtualization/Private/VirtualizationManager.cpp
--- a/Engine/Source/Developer/Virtualization/Private/VirtualizationManager.cpp
+++ b/Engine/Source/Developer/Virtualization/Private/VirtualizationManager.cpp
@@ -578,6 +578,16 @@ EQueryResult FVirtualizationManager::QueryPayloadStatuses(TArrayView<const FIoHa
 	return EQueryResult::Success;
 }
 
+bool FVirtualizationManager::TryVirtualizePackages(const TArray<FString>& FilesToVirtualize, TArray<FText>& OutDescriptionTags, TArray<FText>& OutErrors)
+{
+	OutDescriptionTags.Reset();
+	OutErrors.Reset();
+
+	UE::Virtualization::VirtualizePackages(FilesToVirtualize, OutDescriptionTags, OutErrors);
+
+	return OutErrors.IsEmpty();
+}
+
 FPayloadActivityInfo FVirtualizationManager::GetAccumualtedPayloadActivityInfo() const
 {
 	FPayloadActivityInfo Info;
diff --git a/Engine/Source/Developer/Virtualization/Private/VirtualizationManager.h b/Engine/Source/Developer/Virtualization/Private/VirtualizationManager.h
--- a/Engine/Source/Developer/Virtualization/Private/VirtualizationManager.h
+++ b/Engine/Source/Developer/Virtualization/Private/VirtualizationManager.h
@@ -117,6 +117,8 @@ private:
 
 	virtual EQueryResult QueryPayloadStatuses(TArrayView<const FIoHash> Ids, EStorageType StorageType, TArray<FPayloadStatus>& OutStatuses) override;
 
+	virtual bool TryVirtualizePackages(const TArray<FString>& FilesToVirtualize, TArray<FText>& OutDescriptionTags, TArray<FText>& OutErrors) override;
+
 	virtual FPayloadActivityInfo GetAccumualtedPayloadActivityInfo() const override;
 
 	virtual void GetPayloadActivityInfo( GetPayloadActivityInfoFuncRef ) const override;
diff --git a/Engine/Source/Runtime/Core/Private/Virtualization/VirtualizationSystem.cpp b/Engine/Source/Runtime/Core/Private/Virtualization/VirtualizationSystem.cpp
--- a/Engine/Source/Runtime/Core/Private/Virtualization/VirtualizationSystem.cpp
+++ b/Engine/Source/Runtime/Core/Private/Virtualization/VirtualizationSystem.cpp
@@ -61,6 +61,16 @@ public:
 	{
 	}
 
+	virtual bool TryVirtualizePackages(const TArray<FString>& FilesToVirtualize, TArray<FText>& OutDescriptionTags, TArray<FText>& OutErrors) override
+	{
+		OutDescriptionTags.Reset();
+		OutErrors.Reset();
+
+		OutErrors.Add(FText::FromString(TEXT("Calling ::TryVirtualizePackages on FNullVirtualizationSystem")));
+
+		return false;
+	}
+
 	virtual FPayloadActivityInfo GetAccumualtedPayloadActivityInfo() const override
 	{
 		return FPayloadActivityInfo();
diff --git a/Engine/Source/Runtime/Core/Public/Virtualization/VirtualizationSystem.h b/Engine/Source/Runtime/Core/Public/Virtualization/VirtualizationSystem.h
--- a/Engine/Source/Runtime/Core/Public/Virtualization/VirtualizationSystem.h
+++ b/Engine/Source/Runtime/Core/Public/Virtualization/VirtualizationSystem.h
@@ -255,6 +255,24 @@ public:
 		return QueryPayloadStatuses(Ids, StorageType, OutStatuses) != EQueryResult::Success;
 	}
 
+	/**
+	 * Runs the virtualization process on a set of packages. All of the packages will be parsed and any found to be containing locally stored
+	 * payloads will have them removed but before they are removed they will be pushed to persistent storage.
+	 * 
+	 * @param FilesToVirtualize		An array of file paths to packages that should be virtualized. If a path resolves to a file that is not 
+	 *								a valid package then it will be silently skipped and will not be considered an error.
+	 * @param OutDescriptionTags	The process may produce description tags associated with the packages which will be placed in this array.
+	 *								These tags can be used to improve logging, or be appended to change list descriptions etc. Note that the 
+	 *								array will be emptied before the process.
+	 *								is run and will not contain any pre-existing entries.
+	 * @param OutErrors				Any error encountered during the process will be added here. If any error is added to the array then it
+	 *								can be assumed that the process will return false. Note that the array will be emptied before the process
+	 *								is run and will not contain any pre-existing entries.
+	 * 
+	 * @return True if the process succeeded and false if it did not. If this returns false then OutErrors should contain at least one entry
+	 */
+	virtual bool TryVirtualizePackages(const TArray<FString>& FilesToVirtualize, TArray<FText>& OutDescriptionTags, TArray<FText>& OutErrors) = 0;
+
 	using GetPayloadActivityInfoFuncRef = TFunctionRef<void(const FString& DebugName, const FString& ConfigName, const FPayloadActivityInfo& PayloadInfo)>;
 
 	/** Access profiling info relating to payload activity per backend. Stats will only be collected if ENABLE_COOK_STATS is enabled.*/
