Remove the check for FPackageTrailer::IsEnabled when virtualizing payloads. We don't care if the system is currently disabled, only what the package file itself contains.

#rnx

- The package trailer will be on by default soon and the option to remove it disabled, so this check is a bit pointless anyway.

diff --git a/Engine/Source/Developer/Virtualization/Private/PackageSubmissionChecks.cpp b/Engine/Source/Developer/Virtualization/Private/PackageSubmissionChecks.cpp
--- a/Engine/Source/Developer/Virtualization/Private/PackageSubmissionChecks.cpp
+++ b/Engine/Source/Developer/Virtualization/Private/PackageSubmissionChecks.cpp
@@ -137,12 +137,6 @@ void VirtualizePackages(const TArray<FString>& FilesToSubmit, TArray<FText>& Out
 		return;
 	}
 
-	// Can't virtualize if the payload trailer system is disabled
-	if (!FPackageTrailer::IsEnabled())
-	{
-		return;
-	}
-
 	if (!System.IsPushingEnabled(EStorageType::Persistent))
 	{
 		UE_LOG(LogVirtualization, Verbose, TEXT("Pushing to persistent backend storage is disabled"));
