Use the 'Log' verbosity when logging the file system backends retry count, there is no need to use 'Display' so that the end user sees it every time.

#rnx

diff --git a/Engine/Source/Developer/Virtualization/Private/VirtualizationFileBackend.cpp b/Engine/Source/Developer/Virtualization/Private/VirtualizationFileBackend.cpp
--- a/Engine/Source/Developer/Virtualization/Private/VirtualizationFileBackend.cpp
+++ b/Engine/Source/Developer/Virtualization/Private/VirtualizationFileBackend.cpp
@@ -47,8 +47,8 @@ bool FFileSystemBackend::Initialize(const FString& ConfigEntry)
 	}
 
 	// Now log a summary of the backend settings to make issues easier to diagnose
-	UE_LOG(LogVirtualization, Display, TEXT("[%s] Using path: '%s'"), *GetDebugName(), *RootDirectory);
-	UE_LOG(LogVirtualization, Display, TEXT("[%s] Will retry failed read attempts %d times with a gap of %dms betwen them"), *GetDebugName(), RetryCount, RetryWaitTimeMS);
+	UE_LOG(LogVirtualization, Log, TEXT("[%s] Using path: '%s'"), *GetDebugName(), *RootDirectory);
+	UE_LOG(LogVirtualization, Log, TEXT("[%s] Will retry failed read attempts %d times with a gap of %dms betwen them"), *GetDebugName(), RetryCount, RetryWaitTimeMS);
 
 	return true;
 }
