Allow the source control backend to attempt to retry the downloading of payloads after a failure.

- In some cases people are getting minor connection failures to the perforce server which causes the p4 print command to fail. In these cases it would be better to retry the command rather than failing to get the payload entirely
- Currently the default will be 2 attempts to pull a payload with a 100ms wait between them, this can be changed when setting up the backend graph as optional config values "RetryCount=" and "RetryWaitTime="
-- Note that RetryWaitTime is in ms
- Moved the parsing of the command line to it's own method in the backend.

diff --git a/Engine/Source/Developer/Virtualization/Private/VirtualizationSourceControlBackend.cpp b/Engine/Source/Developer/Virtualization/Private/VirtualizationSourceControlBackend.cpp
--- a/Engine/Source/Developer/Virtualization/Private/VirtualizationSourceControlBackend.cpp
+++ b/Engine/Source/Developer/Virtualization/Private/VirtualizationSourceControlBackend.cpp
@@ -95,6 +95,21 @@ static void CreateDescription(const FString& ProjectName, const TArray<const FPu
 	return SCCProvider.GetState(DepotPaths, OutStates, EStateCacheUsage::Use);
 }
 
+/** Parse all error messages in a FSourceControlResultInfo and return true if the file not found error message is found */
+[[nodiscard]] static bool IsDepotFileMissing(const FSourceControlResultInfo& ResultInfo)
+{
+	// Ideally we'd parse for this sort of thing in the source control module itself and return an error enum
+	for (const FText& ErrorTest : ResultInfo.ErrorMessages)
+	{
+		if (ErrorTest.ToString().Find(" - no such file(s).") != INDEX_NONE)
+		{
+			return true;
+		}
+	}
+
+	return false;
+}
+
 FSourceControlBackend::FSourceControlBackend(FStringView ProjectName, FStringView ConfigName, FStringView InDebugName)
 	: IVirtualizationBackend(ConfigName, InDebugName, EOperations::Both)
 	, ProjectName(ProjectName)
@@ -105,22 +120,11 @@ bool FSourceControlBackend::Initialize(const FString& ConfigEntry)
 {
 	TRACE_CPUPROFILER_EVENT_SCOPE(FSourceControlBackend::Initialize);
 
-	// We require that a valid depot root has been provided
-	if (!FParse::Value(*ConfigEntry, TEXT("DepotRoot="), DepotRoot))
-	{
-		UE_LOG(LogVirtualization, Error, TEXT("'DepotRoot=' not found in the config file"));
-		return false;
-	}
-
-	if (!FindSubmissionWorkingDir(ConfigEntry))
+	if (!TryApplySettingsFromConfigFiles(ConfigEntry))
 	{
 		return false;
 	}
 
-	// Optional config values
-	FParse::Bool(*ConfigEntry, TEXT("UsePartitionedClient="), bUsePartitionedClient);
-	UE_LOG(LogVirtualization, Display, TEXT("[%s] Using partitioned clients: '%s'"), *GetDebugName(), bUsePartitionedClient ? TEXT("true") : TEXT("false"));
-
 	// We do not want the connection to have a client workspace so explicitly set it to empty
 	FSourceControlInitSettings SCCSettings(FSourceControlInitSettings::EBehavior::OverrideExisting);
 	SCCSettings.AddSetting(TEXT("P4Client"), TEXT(""));
@@ -211,24 +215,46 @@ FCompressedBuffer FSourceControlBackend::PullData(const FIoHash& Id)
 
 	UE_LOG(LogVirtualization, Verbose, TEXT("[%s] Attempting to pull '%s' from source control"), *GetDebugName(), *DepotPath);
 
-#if IS_SOURCE_CONTROL_THREAD_SAFE
-	TSharedRef<FDownloadFile, ESPMode::ThreadSafe> DownloadCommand = ISourceControlOperation::Create<FDownloadFile>(FDownloadFile::EVerbosity::None);
-	if (SCCProvider->Execute(DownloadCommand, DepotPath.ToString(), EConcurrency::Synchronous) != ECommandResult::Succeeded)
+	int32 Retries = 0;
+
+	while (Retries < RetryCount)
 	{
-		return FCompressedBuffer();
-	}
+		// Only warn if the backend is configured to retry
+		if (Retries != 0)
+		{
+			UE_LOG(LogVirtualization, Warning, TEXT("[%s] Failed to download '%s' retrying (%d/%d) in %dms..."), *GetDebugName(), DepotPath.ToString(), Retries, RetryCount, RetryWaitTimeMS);
+			FPlatformProcess::SleepNoStats(RetryWaitTimeMS * 0.001f);
+		}
+
+#if IS_SOURCE_CONTROL_THREAD_SAFE
+		TSharedRef<FDownloadFile, ESPMode::ThreadSafe> DownloadCommand = ISourceControlOperation::Create<FDownloadFile>(FDownloadFile::EVerbosity::None);
+		if (SCCProvider->Execute(DownloadCommand, DepotPath.ToString(), EConcurrency::Synchronous) == ECommandResult::Succeeded)
+		{
+			// The payload was created by FCompressedBuffer::Compress so we can return it as a FCompressedBuffer.
+			FSharedBuffer Buffer = DownloadCommand->GetFileData(DepotPath);
+			return FCompressedBuffer::FromCompressed(Buffer);
+		}
 #else
-	TSharedRef<FDownloadFile> DownloadCommand = ISourceControlOperation::Create<FDownloadFile>(FDownloadFile::EVerbosity::None);
-	if (!SCCProvider->TryToDownloadFileFromBackgroundThread(DownloadCommand, DepotPath.ToString()))
-	{
-		return FCompressedBuffer();
-	}
+		TSharedRef<FDownloadFile> DownloadCommand = ISourceControlOperation::Create<FDownloadFile>(FDownloadFile::EVerbosity::None);
+		if (SCCProvider->TryToDownloadFileFromBackgroundThread(DownloadCommand, DepotPath.ToString()))
+		{
+			// The payload was created by FCompressedBuffer::Compress so we can return it as a FCompressedBuffer.
+			FSharedBuffer Buffer = DownloadCommand->GetFileData(DepotPath);
+			return FCompressedBuffer::FromCompressed(Buffer);
+		}
 #endif
 
-	// The payload was created by FCompressedBuffer::Compress so we can return it 
-	// as a FCompressedBuffer.
-	FSharedBuffer Buffer = DownloadCommand->GetFileData(DepotPath);
-	return FCompressedBuffer::FromCompressed(Buffer);
+		// If this was the first try then check to see if the error being returns is that the file does not exist
+		// in the depot. If it does not exist then there is no point in us retrying and we can error out at this point.
+		if (Retries == 0 && IsDepotFileMissing(DownloadCommand->GetResultInfo()))
+		{
+			return FCompressedBuffer();
+		}
+		
+		Retries++;
+	}
+
+	return FCompressedBuffer();
 }
 
 bool FSourceControlBackend::DoesPayloadExist(const FIoHash& Id)
@@ -536,6 +562,42 @@ bool FSourceControlBackend::DoPayloadsExist(TArrayView<const FIoHash> PayloadIds
 	return true;
 }
 
+bool FSourceControlBackend::TryApplySettingsFromConfigFiles(const FString& ConfigEntry)
+{
+	// We require that a valid depot root has been provided
+	if (!FParse::Value(*ConfigEntry, TEXT("DepotRoot="), DepotRoot))
+	{
+		UE_LOG(LogVirtualization, Error, TEXT("'DepotRoot=' not found in the config file"));
+		return false;
+	}
+
+	// Optional config values
+	FParse::Bool(*ConfigEntry, TEXT("UsePartitionedClient="), bUsePartitionedClient);
+
+	int32 RetryCountIniFile = INDEX_NONE;
+	if (FParse::Value(*ConfigEntry, TEXT("RetryCount="), RetryCountIniFile))
+	{
+		RetryCount = RetryCountIniFile;
+	}
+
+	int32 RetryWaitTimeMSIniFile = INDEX_NONE;
+	if (FParse::Value(*ConfigEntry, TEXT("RetryWaitTime="), RetryWaitTimeMSIniFile))
+	{
+		RetryWaitTimeMS = RetryWaitTimeMSIniFile;
+	}
+
+	// Now log a summary of the optional settings to make issues easier to diagnose
+	UE_LOG(LogVirtualization, Log, TEXT("[%s] Using partitioned clients: '%s'"), *GetDebugName(), bUsePartitionedClient ? TEXT("true") : TEXT("false"));
+	UE_LOG(LogVirtualization, Log, TEXT("[%s] Will retry failed downloads attempts %d time(s) with a gap of %dms betwen them"), *GetDebugName(), RetryCount, RetryWaitTimeMS);
+
+	if (!FindSubmissionWorkingDir(ConfigEntry))
+	{
+		return false;
+	}
+
+	return true;
+}
+
 void FSourceControlBackend::CreateDepotPath(const FIoHash& PayloadId, FStringBuilderBase& OutPath)
 {
 	TStringBuilder<52> PayloadPath;
@@ -555,7 +617,7 @@ bool FSourceControlBackend::FindSubmissionWorkingDir(const FString& ConfigEntry)
 	if (!SubmissionRootDir.IsEmpty())
 	{
 		FPaths::NormalizeDirectoryName(SubmissionRootDir);
-		UE_LOG(LogVirtualization, Display, TEXT("[%s] Found Environment Variable: UE-VirtualizationWorkingDir"), *GetDebugName());	
+		UE_LOG(LogVirtualization, Log, TEXT("[%s] Found Environment Variable: UE-VirtualizationWorkingDir"), *GetDebugName());	
 	}
 	else
 	{
@@ -578,7 +640,7 @@ bool FSourceControlBackend::FindSubmissionWorkingDir(const FString& ConfigEntry)
 
 	if (IFileManager::Get().DirectoryExists(*SubmissionRootDir) || IFileManager::Get().MakeDirectory(*SubmissionRootDir))
 	{
-		UE_LOG(LogVirtualization, Display, TEXT("[%s] Setting '%s' as the working directory"), *GetDebugName(), *SubmissionRootDir);
+		UE_LOG(LogVirtualization, Log, TEXT("[%s] Setting '%s' as the working directory"), *GetDebugName(), *SubmissionRootDir);
 		return true;
 	}
 	else
diff --git a/Engine/Source/Developer/Virtualization/Private/VirtualizationSourceControlBackend.h b/Engine/Source/Developer/Virtualization/Private/VirtualizationSourceControlBackend.h
--- a/Engine/Source/Developer/Virtualization/Private/VirtualizationSourceControlBackend.h
+++ b/Engine/Source/Developer/Virtualization/Private/VirtualizationSourceControlBackend.h
@@ -31,6 +31,13 @@ namespace UE::Virtualization
  * SubmitFromTempDir [bool]:	When set to true, payloads will be submitted from the temp directory of
  *								the current machine and when false the files will be submitted from the
  *								Save directory of the current project. [Default=false]
+ * RetryCount [int32]:			How many times we should try to download a payload before giving up with
+ *								an error. Useful when the connection is unreliable but does not experience 
+ *								frequent persistent outages. [Default=2]
+ * RetryWaitTime [int32]:		The length of time the process should wait between each download attempt
+ *								in milliseconds. Remember that the max length of time that the process
+ *								can stall attempting to download a payload file is 
+ *								RetryCount * RetryWaitTime; [Default=100ms]
  * 
  * Environment Variables:
  * UE-VirtualizationWorkingDir [string]:	This can be set to a valid directory path that the backend
@@ -60,6 +67,8 @@ private:
 
 private:
 
+	bool TryApplySettingsFromConfigFiles(const FString& ConfigEntry);
+
 	void CreateDepotPath(const FIoHash& PayloadId, FStringBuilderBase& OutPath);
 
 	bool FindSubmissionWorkingDir(const FString& ConfigEntry);
@@ -82,6 +91,11 @@ private:
 	/** Should we try to make the temp client partitioned or not? */
 	bool bUsePartitionedClient = true;
 
+	/** The number of times to retry pulling a payload from the depot */
+	int32 RetryCount = 2;
+
+	/** The length of time (in milliseconds) to wait after each pull attempt before retrying. */
+	int32 RetryWaitTimeMS = 100;
 };
 
 } // namespace UE::Virtualization
