Extend VA.MissBackends commandline to work via console commands as well.

#rnx

- The commandline -VA-MissBackends has been extended to a console command 'VA.MissBackends' that can be used at runtime to disable specific backends or reset the list and enable all of them again.
-- Entering VA.MissBackends without any args will print help documentation to the console/output log and show the user the names of the backends that they can disable.
- Technically we don't need the command line parsing anymore and users could use -dpcvars=VA.MissBackends=X instead but since we already have the code we might as well leave it untouched.
- Disabling backends from pulling payloads for debug purposes has been split off to it's own system stored as a

- IVirtualizationBackend::EOperations has been turned into bit flags, so we no longer need the 'both' entry
- IVirtualizationBackendnow has two EOperations members, one to record the operations that the backend supports and another that allows the operations to be disabled for debugging purposes.
- IVirtualizationBackend now has one method to poll if an operation is supported ::IsOperationSupported, which should not change over the lift time of the backend. In addition there is also a method for disabling/enabling a backend for debugging purposes as well as polling the debug state.

diff --git a/Engine/Source/Developer/Virtualization/Private/VirtualizationDDCBackend.cpp b/Engine/Source/Developer/Virtualization/Private/VirtualizationDDCBackend.cpp
--- a/Engine/Source/Developer/Virtualization/Private/VirtualizationDDCBackend.cpp
+++ b/Engine/Source/Developer/Virtualization/Private/VirtualizationDDCBackend.cpp
@@ -18,7 +18,7 @@ static UE::DerivedData::FValueId ToDerivedDataValueId(const FIoHash& Id)
 }
 
 FDDCBackend::FDDCBackend(FStringView ProjectName, FStringView ConfigName, FStringView InDebugName)
-: IVirtualizationBackend(ConfigName, InDebugName, EOperations::Both)
+: IVirtualizationBackend(ConfigName, InDebugName, EOperations::Push | EOperations::Pull)
 , BucketName(TEXT("BulkData"))
 , TransferPolicy(UE::DerivedData::ECachePolicy::None)
 , QueryPolicy(UE::DerivedData::ECachePolicy::None)
diff --git a/Engine/Source/Developer/Virtualization/Private/VirtualizationFileBackend.cpp b/Engine/Source/Developer/Virtualization/Private/VirtualizationFileBackend.cpp
--- a/Engine/Source/Developer/Virtualization/Private/VirtualizationFileBackend.cpp
+++ b/Engine/Source/Developer/Virtualization/Private/VirtualizationFileBackend.cpp
@@ -12,7 +12,7 @@ namespace UE::Virtualization
 {
 
 FFileSystemBackend::FFileSystemBackend(FStringView ProjectName, FStringView ConfigName, FStringView DebugName)
-	: IVirtualizationBackend(ConfigName, DebugName, EOperations::Both)
+	: IVirtualizationBackend(ConfigName, DebugName, EOperations::Push | EOperations::Pull)
 {
 }
 
diff --git a/Engine/Source/Developer/Virtualization/Private/VirtualizationHttpBackend.cpp b/Engine/Source/Developer/Virtualization/Private/VirtualizationHttpBackend.cpp
--- a/Engine/Source/Developer/Virtualization/Private/VirtualizationHttpBackend.cpp
+++ b/Engine/Source/Developer/Virtualization/Private/VirtualizationHttpBackend.cpp
@@ -1294,7 +1294,7 @@ private:
 };
 
 FHttpBackend::FHttpBackend(FStringView ProjectName, FStringView ConfigName, FStringView InDebugName)
-	: IVirtualizationBackend(ConfigName, InDebugName, EOperations::Both)
+	: IVirtualizationBackend(ConfigName, InDebugName, EOperations::Push | EOperations::Pull)
 	, Namespace(TEXT("mirage"))
 	, Bucket(TEXT("default"))
 	, ChunkSize(-1)
diff --git a/Engine/Source/Developer/Virtualization/Private/VirtualizationManager.cpp b/Engine/Source/Developer/Virtualization/Private/VirtualizationManager.cpp
--- a/Engine/Source/Developer/Virtualization/Private/VirtualizationManager.cpp
+++ b/Engine/Source/Developer/Virtualization/Private/VirtualizationManager.cpp
@@ -269,12 +269,22 @@ FVirtualizationManager::FVirtualizationManager()
 	// TODO: We should just be able to call the logging in the destructor, but 
 	// we need to fix the startup/shutdown ordering of Mirage first.
 	COOK_STAT(FCoreDelegates::OnExit.AddStatic(Profiling::LogStats));
+
+	DebugConsoleCommands.Add(IConsoleManager::Get().RegisterConsoleCommand(
+		TEXT("VA.MissBackends"),
+		TEXT("A debug commnad which can be used to disable payload pulling on one or more backends"),
+		FConsoleCommandWithArgsAndOutputDeviceDelegate::CreateRaw(this, &FVirtualizationManager::OnUpdateMissBackendsFromConsole)));
 }
 
 FVirtualizationManager::~FVirtualizationManager()
 {
 	UE_LOG(LogVirtualization, Log, TEXT("Destroying backends"));
 	
+	for (IConsoleCommand* Cmd : DebugConsoleCommands)
+	{
+		IConsoleManager::Get().UnregisterConsoleObject(Cmd);
+	}
+
 	LocalCachableBackends.Empty();
 	PersistentStorageBackends.Empty();
 	PullEnabledBackends.Empty();
@@ -425,7 +435,7 @@ bool FVirtualizationManager::PushData(TArrayView<FPushRequest> Requests, EStorag
 		
 		// Debugging operation where we immediately try to pull the payload after each push (when possible) and assert 
 		// that the pulled payload is the same as the original
-		if (bValidateAfterPushOperation && bResult == true && Backend->SupportsPullOperations())
+		if (bValidateAfterPushOperation && bResult == true && Backend->IsOperationSupported(IVirtualizationBackend::EOperations::Pull))
 		{
 			for (FPushRequest& Request : ValidatedRequests)
 			{
@@ -479,6 +489,12 @@ FCompressedBuffer FVirtualizationManager::PullData(const FIoHash& Id)
 
 	for (IVirtualizationBackend* Backend : PullEnabledBackends)
 	{
+		// Skip if pulling has been disabled on this backend for debug purposes
+		if (Backend->IsOperationDebugDisabled(IVirtualizationBackend::EOperations::Pull))
+		{
+			continue;
+		}
+
 		FCompressedBuffer Payload = PullDataFromBackend(*Backend, Id);
 
 		if (Payload)
@@ -786,16 +802,112 @@ void FVirtualizationManager::ApplyDebugSettingsFromFromCmdline()
 	FString MissOptions;
 	if (FParse::Value(FCommandLine::Get(), TEXT("-VA-MissBackends="), MissOptions))
 	{
-		MissOptions.ParseIntoArray(BackendsToDisablePulls, TEXT("+"), true);
+		MissOptions.ParseIntoArray(DebugMissBackends, TEXT("+"), true);
 
 		UE_LOG(LogVirtualization, Warning, TEXT("Cmdline has disabled payload pulling for the following backends:"));
-		for (const FString& Backend : BackendsToDisablePulls)
+		for (const FString& Backend : DebugMissBackends)
 		{
 			UE_LOG(LogVirtualization, Warning, TEXT("\t%s"), *Backend);
 		}
 	}
 }
 
+void FVirtualizationManager::OnUpdateMissBackendsFromConsole(const TArray<FString>& Args, FOutputDevice& OutputDevice)
+{
+	if (Args.IsEmpty())
+	{
+		OutputDevice.Log(TEXT("VA.MissBackends command help"));
+		OutputDevice.Log(TEXT("This command allows you to disable the pulling of payloads by specific backends"));
+		OutputDevice.Log(TEXT(""));
+		OutputDevice.Log(TEXT("Commands:"));
+		OutputDevice.Log(TEXT("VA.MissBackends reset            - Empties the list of backends, everything will function normally"));
+		OutputDevice.Log(TEXT("VA.MissBackends list             - Prints the list of backends affected"));
+		OutputDevice.Log(TEXT("VA.MissBackends set Name0 Name1  - List each backend that you want to fail to pull payloads"));
+		OutputDevice.Log(TEXT("VA.MissBackends set All          - All backends will fail to pull payloads"));
+		OutputDevice.Log(TEXT(""));
+		OutputDevice.Log(TEXT("Valid backend names:"));
+
+		for (const TUniquePtr<IVirtualizationBackend>& Backend : AllBackends)
+		{
+			OutputDevice.Logf(TEXT("\t%s"), *Backend->GetConfigName());
+		}
+	}
+	else if (Args.Num() == 1)
+	{
+		if (Args[0] == TEXT("reset"))
+		{
+			DebugMissBackends.Empty();
+			UpdateBackendDebugState();
+		}
+		else if (Args[0] == TEXT("list"))
+		{
+			if (!DebugMissBackends.IsEmpty())
+			{
+				OutputDevice.Log(TEXT("Disabled backends:"));
+				for (const FString& Backend : DebugMissBackends)
+				{
+					OutputDevice.Logf(TEXT("\t%s"), *Backend);
+				}
+			}
+			else
+			{
+				OutputDevice.Log(TEXT("No backends are disabled"));
+			}
+		}
+		else
+		{
+			OutputDevice.Log(ELogVerbosity::Error, TEXT("Invalid args for the VA.MissBackends command!"));
+		}
+	}
+	else if (Args[0] == TEXT("set"))
+	{	
+		DebugMissBackends.Empty(Args.Num() - 1);
+
+		for (int32 Index = 1; Index < Args.Num(); ++Index)
+		{
+			DebugMissBackends.Add(Args[Index]);
+		}
+
+		UpdateBackendDebugState();
+	}
+	else
+	{
+		OutputDevice.Log(ELogVerbosity::Error, TEXT("Invalid args for the VA.MissBackends command!"));
+	}
+}
+
+void FVirtualizationManager::UpdateBackendDebugState()
+{
+	for (TUniquePtr<IVirtualizationBackend>& Backend : AllBackends)
+	{
+		const bool bDisable = ShouldDebugDisablePulling(Backend->GetConfigName());
+		Backend->SetOperationDebugState(IVirtualizationBackend::EOperations::Pull, bDisable);
+	}
+}
+
+bool FVirtualizationManager::ShouldDebugDisablePulling(FStringView BackendConfigName) const
+{
+	if (DebugMissBackends.IsEmpty())
+	{
+		return false;
+	}
+
+	if (DebugMissBackends[0] == TEXT("All"))
+	{
+		return true;
+	}
+
+	for (const FString& Name : DebugMissBackends)
+	{
+		if (Name == BackendConfigName)
+		{
+			return true;
+		}
+	}
+
+	return false;
+}
+
 void FVirtualizationManager::MountBackends(const FConfigFile& ConfigFile)
 {
 	TRACE_CPUPROFILER_EVENT_SCOPE(FVirtualizationManager::MountBackends);
@@ -816,6 +928,9 @@ void FVirtualizationManager::MountBackends(const FConfigFile& ConfigFile)
 	// persistent storage backends in 'PullEnabledBackends'.
 	ParseHierarchy(ConfigFile, GraphName, TEXT("LocalStorageHierarchy"), FactoryLookupTable, LocalCachableBackends);
 	ParseHierarchy(ConfigFile, GraphName, TEXT("PersistentStorageHierarchy"), FactoryLookupTable, PersistentStorageBackends);
+
+	// Apply and disabled backends from the command line
+	UpdateBackendDebugState();
 }
 
 void FVirtualizationManager::ParseHierarchy(const FConfigFile& ConfigFile, const TCHAR* GraphName, const TCHAR* HierarchyKey, const FRegistedFactories& FactoryLookupTable, FBackendArray& PushArray)
@@ -874,11 +989,6 @@ bool FVirtualizationManager::CreateBackend(const FConfigFile& ConfigFile, const
 
 			}
 
-			if (BackendsToDisablePulls.Find(Backend->GetConfigName()) != INDEX_NONE  || BackendsToDisablePulls.Find(TEXT("All")) != INDEX_NONE)
-			{
-				Backend->DisablePullOperationSupport();
-			}
-			
 			if (Backend->Initialize(Cmdine))
 			{
 				AddBackend(MoveTemp(Backend), PushArray);
@@ -914,12 +1024,12 @@ void FVirtualizationManager::AddBackend(TUniquePtr<IVirtualizationBackend> Backe
 	// Get a reference pointer to use in the other backend arrays
 	IVirtualizationBackend* BackendRef = AllBackends.Last().Get();
 
-	if (BackendRef->SupportsPullOperations())
+	if (BackendRef->IsOperationSupported(IVirtualizationBackend::EOperations::Pull))
 	{
 		PullEnabledBackends.Add(BackendRef);
 	}
 
-	if (BackendRef->SupportsPushOperations())
+	if (BackendRef->IsOperationSupported(IVirtualizationBackend::EOperations::Push))
 	{
 		PushArray.Add(BackendRef);
 	}
@@ -947,7 +1057,7 @@ void FVirtualizationManager::CachePayload(const FIoHash& Id, const FCompressedBu
 
 		// Debugging operation where we immediately try to pull the payload after each push (when possible) and assert 
 		// that the pulled payload is the same as the original
-		if (bValidateAfterPushOperation && bResult && BackendToCache->SupportsPullOperations())
+		if (bValidateAfterPushOperation && bResult && BackendToCache->IsOperationSupported(IVirtualizationBackend::EOperations::Pull))
 		{
 			FCompressedBuffer PulledPayload = PullDataFromBackend(*BackendToCache, Id);
 			checkf(	Payload.GetRawHash() == PulledPayload.GetRawHash(), 
diff --git a/Engine/Source/Developer/Virtualization/Private/VirtualizationManager.h b/Engine/Source/Developer/Virtualization/Private/VirtualizationManager.h
--- a/Engine/Source/Developer/Virtualization/Private/VirtualizationManager.h
+++ b/Engine/Source/Developer/Virtualization/Private/VirtualizationManager.h
@@ -9,6 +9,9 @@
 
 #include "Virtualization/VirtualizationSystem.h"
 
+struct IConsoleCommand;
+class FOutputDevice;
+
 /**
  * Configuring the backend hierarchy
  * 
@@ -136,6 +139,10 @@ private:
 	void ApplyDebugSettingsFromConfigFiles(const FConfigFile& ConfigFile);
 	void ApplyDebugSettingsFromFromCmdline();
 
+	void OnUpdateMissBackendsFromConsole(const TArray<FString>& Args, FOutputDevice& OutputDevice);
+	void UpdateBackendDebugState();
+	bool ShouldDebugDisablePulling(FStringView BackendConfigName) const;	
+
 	void MountBackends(const FConfigFile& ConfigFile);
 	void ParseHierarchy(const FConfigFile& ConfigFile, const TCHAR* GraphName, const TCHAR* HierarchyKey, const FRegistedFactories& FactoryLookupTable, FBackendArray& PushArray);
 	bool CreateBackend(const FConfigFile& ConfigFile, const TCHAR* GraphName, const FString& ConfigEntryName, const FRegistedFactories& FactoryLookupTable, FBackendArray& PushArray);
@@ -210,9 +217,6 @@ private:
 	/** The name of the current project */
 	FString ProjectName;
 
-	/** Array of backend names that should have their pull operation disabled */
-	TArray<FString> BackendsToDisablePulls;
-
 	/** The critical section used to force single threaded access if bForceSingleThreaded is true */
 	FCriticalSection ForceSingleThreadedCS;
 
@@ -233,6 +237,14 @@ private:
 
 	/** Our notification Event */
 	FOnNotification NotificationEvent;
+
+	// Members after this point at used for debugging operations only!
+
+	/** Console commands that the manager has registered */
+	TArray<IConsoleCommand*> DebugConsoleCommands;
+	
+	/** Array of backend names that should have their pull operation disabled */
+	TArray<FString> DebugMissBackends;
 };
 
 } // namespace UE::Virtualization
diff --git a/Engine/Source/Developer/Virtualization/Private/VirtualizationSourceControlBackend.cpp b/Engine/Source/Developer/Virtualization/Private/VirtualizationSourceControlBackend.cpp
--- a/Engine/Source/Developer/Virtualization/Private/VirtualizationSourceControlBackend.cpp
+++ b/Engine/Source/Developer/Virtualization/Private/VirtualizationSourceControlBackend.cpp
@@ -111,7 +111,7 @@ static void CreateDescription(const FString& ProjectName, const TArray<const FPu
 }
 
 FSourceControlBackend::FSourceControlBackend(FStringView ProjectName, FStringView ConfigName, FStringView InDebugName)
-	: IVirtualizationBackend(ConfigName, InDebugName, EOperations::Both)
+	: IVirtualizationBackend(ConfigName, InDebugName, EOperations::Push | EOperations::Pull)
 	, ProjectName(ProjectName)
 {
 }
diff --git a/Engine/Source/Developer/Virtualization/Public/IVirtualizationBackend.h b/Engine/Source/Developer/Virtualization/Public/IVirtualizationBackend.h
--- a/Engine/Source/Developer/Virtualization/Public/IVirtualizationBackend.h
+++ b/Engine/Source/Developer/Virtualization/Public/IVirtualizationBackend.h
@@ -37,26 +37,29 @@ enum class EPushResult
  */
 class IVirtualizationBackend
 {
-protected:
+public: 
 	/** Enum detailing which operations a backend can support */
 	enum class EOperations : uint8
 	{
 		/** Supports no operations, this should only occur when debug settings are applied */
 		None = 0,
 		/** Supports only push operations */
-		Push,
+		Push = 1 << 0,
 		/** Supports only pull operations */
-		Pull,
-		/** Supports both push and pull operations */
-		Both
+		Pull = 1 << 1,
 	};
 
+	FRIEND_ENUM_CLASS_FLAGS(EOperations);
+
+protected:
+	
 	IVirtualizationBackend(FStringView InConfigName, FStringView InDebugName, EOperations InSupportedOperations)
 		: SupportedOperations(InSupportedOperations)
+		, DebugDisabledOperations(EOperations::None)
 		, ConfigName(InConfigName)
 		, DebugName(InDebugName)
 	{
-		checkf(InSupportedOperations != EOperations::None, TEXT("Cannot create a backend without supporting at least one type of operation!"));
+		checkf(InSupportedOperations != EOperations::None, TEXT("Cannot create a backend with no supported operations!"));
 	}
 
 public:
@@ -164,29 +167,33 @@ public:
 		return true;
 	}
 	
-	/** Used when debugging to disable the pull operation */
-	void DisablePullOperationSupport()
+	/** 
+	 * Returns true if the given operation is supported, this is set when the backend is created
+	 * and should not change over it's life time.
+	 */
+	bool IsOperationSupported(EOperations Operation) const
+	{
+		return EnumHasAnyFlags(SupportedOperations, Operation);
+	}
+
+	/** Enable or disable the given operation based on the 'bIsDisabled' parameter */
+	void SetOperationDebugState(EOperations Operation, bool bIsDisabled)
 	{
-		if (SupportedOperations == EOperations::Pull)
+		if (bIsDisabled)
 		{
-			SupportedOperations = EOperations::None;
+			EnumAddFlags(DebugDisabledOperations, Operation);
+			
 		}
-		else if (SupportedOperations == EOperations::Both)
+		else
 		{
-			SupportedOperations = EOperations::Push;
+			EnumRemoveFlags(DebugDisabledOperations, Operation);
 		}
 	}
 
-	/** Return true if the backend supports push operations. Returning true allows ::PushData to be called.  */
-	bool SupportsPushOperations() const 
-	{ 
-		return SupportedOperations == EOperations::Push || SupportedOperations == EOperations::Both;  
-	}
-
-	/** Return true if the backend supports pull operations. Returning true allows ::PullData to be called.  */
-	bool SupportsPullOperations() const
-	{ 
-		return SupportedOperations == EOperations::Pull || SupportedOperations == EOperations::Both;  
+	/** Returns true if the given operation is disabled for debugging purposes */
+	bool IsOperationDebugDisabled(EOperations Operation) const
+	{
+		return EnumHasAnyFlags(DebugDisabledOperations, Operation);
 	}
 
 	/** Returns a string containing the name of the backend as it appears in the virtualization graph in the config file */
@@ -206,6 +213,8 @@ private:
 	/** The operations that this backend supports */
 	EOperations SupportedOperations;
 
+	EOperations DebugDisabledOperations;
+
 	/** The name assigned to the backend by the virtualization graph */
 	FString ConfigName;
 
@@ -237,6 +246,8 @@ public:
 	virtual FName GetName() = 0;
 };
 
+ENUM_CLASS_FLAGS(IVirtualizationBackend::EOperations);
+
 /**
  * This macro is used to generate a backend factories boilerplate code if you do not
  * need anything more than the default behavior.
