Add new commandline '-VA-MissPercent=' and console command 'VA.MissChance' that allows the user to set a chance (in percent) that a virtualized payload pull will just fail. This is to allow QA to more easily simulate bad and flaky connections.

#rnx

- Added verbose logging to show when a payload is skipped due to the miss chance, or when a backend skips a payload due to the 'VA.MissBackends' command

diff --git a/Engine/Source/Developer/Virtualization/Private/VirtualizationManager.cpp b/Engine/Source/Developer/Virtualization/Private/VirtualizationManager.cpp
--- a/Engine/Source/Developer/Virtualization/Private/VirtualizationManager.cpp
+++ b/Engine/Source/Developer/Virtualization/Private/VirtualizationManager.cpp
@@ -274,6 +274,11 @@ FVirtualizationManager::FVirtualizationManager()
 		TEXT("VA.MissBackends"),
 		TEXT("A debug commnad which can be used to disable payload pulling on one or more backends"),
 		FConsoleCommandWithArgsAndOutputDeviceDelegate::CreateRaw(this, &FVirtualizationManager::OnUpdateMissBackendsFromConsole)));
+
+	DebugConsoleCommands.Add(IConsoleManager::Get().RegisterConsoleCommand(
+		TEXT("VA.MissChance"),
+		TEXT("A debug command which can be used to set the chance that a payload pull will fail"),
+		FConsoleCommandWithArgsAndOutputDeviceDelegate::CreateRaw(this, &FVirtualizationManager::OnUpdateMissChanceFromConsole)));
 }
 
 FVirtualizationManager::~FVirtualizationManager()
@@ -487,42 +492,27 @@ FCompressedBuffer FVirtualizationManager::PullData(const FIoHash& Id)
 
 	GetNotificationEvent().Broadcast(IVirtualizationSystem::PullBegunNotification, Id);
 
-	for (IVirtualizationBackend* Backend : PullEnabledBackends)
-	{
-		// Skip if pulling has been disabled on this backend for debug purposes
-		if (Backend->IsOperationDebugDisabled(IVirtualizationBackend::EOperations::Pull))
-		{
-			continue;
-		}
-
-		FCompressedBuffer Payload = PullDataFromBackend(*Backend, Id);
-
-		if (Payload)
-		{
-			if (bEnableCacheAfterPull)
-			{
-				CachePayload(Id, Payload, Backend);
-			}
-
-			GetNotificationEvent().Broadcast(IVirtualizationSystem::PullEndedNotification, Id);
-
-			return Payload;
-		}
-	}
-
-	// Pull ended but failed..
+	FCompressedBuffer Payload = PullDataFromAllBackends(Id);
+	
 	GetNotificationEvent().Broadcast(IVirtualizationSystem::PullEndedNotification, Id);
 
-	// Broadcast the pull failed event to any listeners
-	GetNotificationEvent().Broadcast(IVirtualizationSystem::PullFailedNotification, Id);
+	if (!Payload.IsNull())
+	{
+		return Payload;
+	}
+	else
+	{
+		// Broadcast the pull failed event to any listeners
+		GetNotificationEvent().Broadcast(IVirtualizationSystem::PullFailedNotification, Id);
 
-	// TODO: Maybe this should be a fatal error? If we keep it as an error we need to make sure any calling
-	// code handles it properly.
-	// Could be worth extending ::PullData to return error codes instead so we can make a better distinction 
-	// between the payload not being found in any of the backends and one or more of the backends failing.
-	UE_LOG(LogVirtualization, Error, TEXT("Payload '%s' failed to be pulled from any backend'"), *LexToString(Id));
+		// TODO: Maybe this should be a fatal error? If we keep it as an error we need to make sure any calling
+		// code handles it properly.
+		// Could be worth extending ::PullData to return error codes instead so we can make a better distinction 
+		// between the payload not being found in any of the backends and one or more of the backends failing.
+		UE_LOG(LogVirtualization, Error, TEXT("Payload '%s' failed to be pulled from any backend'"), *LexToString(Id));
 
-	return FCompressedBuffer();
+		return FCompressedBuffer();
+	}
 }
 
 EQueryResult FVirtualizationManager::QueryPayloadStatuses(TArrayView<const FIoHash> Ids, EStorageType StorageType, TArray<FPayloadStatus>& OutStatuses)
@@ -810,6 +800,14 @@ void FVirtualizationManager::ApplyDebugSettingsFromFromCmdline()
 			UE_LOG(LogVirtualization, Warning, TEXT("\t%s"), *Backend);
 		}
 	}
+
+	DebugMissChance = 0.0f;
+	if (FParse::Value(FCommandLine::Get(), TEXT("-VA-MissChance="), DebugMissChance))
+	{
+		DebugMissChance = FMath::Clamp(DebugMissChance, 0.0f, 100.0f);
+
+		UE_LOG(LogVirtualization, Warning, TEXT("Cmdline has set a %.1f%% chance of a payload pull failing"), DebugMissChance);
+	}
 }
 
 void FVirtualizationManager::OnUpdateMissBackendsFromConsole(const TArray<FString>& Args, FOutputDevice& OutputDevice)
@@ -876,6 +874,40 @@ void FVirtualizationManager::OnUpdateMissBackendsFromConsole(const TArray<FStrin
 	}
 }
 
+void FVirtualizationManager::OnUpdateMissChanceFromConsole(const TArray<FString>& Args, FOutputDevice& OutputDevice)
+{
+	if (Args.IsEmpty())
+	{
+		OutputDevice.Log(TEXT("VA.MissChance command help"));
+		OutputDevice.Log(TEXT("This command allows you to set the chance (in percent) that a payload pull request will just fail"));
+		OutputDevice.Log(TEXT(""));
+		OutputDevice.Log(TEXT("Commands:"));
+		OutputDevice.Log(TEXT("VA.MissChance show     - prints the current miss percent chance"));
+		OutputDevice.Log(TEXT("VA.MissChance set Num - Sets the miss percent chance to the given value"));
+	}
+	else if (Args.Num() == 1 && Args[0] == TEXT("show"))
+	{
+		OutputDevice.Logf(TEXT("Current debug miss chance: %.1f%%"), DebugMissChance);
+	}
+	else if (Args.Num() == 2 && Args[0] == TEXT("set"))
+	{
+		if (::LexTryParseString(DebugMissChance, *Args[1]))
+		{
+			DebugMissChance = FMath::Clamp(DebugMissChance, 0.0f, 100.0f);
+			OutputDevice.Logf(TEXT("Current debug miss chance set to %.1f%%"), DebugMissChance);
+		}
+		else
+		{
+			DebugMissChance = 0.0f;
+			OutputDevice.Log(ELogVerbosity::Error, TEXT("Invalid value, current debug miss chance reset to 0.0%"));
+		}
+	}
+	else
+	{
+		OutputDevice.Log(ELogVerbosity::Error, TEXT("Invalid args for the VA.MissChance command!"));
+	}
+}
+
 void FVirtualizationManager::UpdateBackendDebugState()
 {
 	for (TUniquePtr<IVirtualizationBackend>& Backend : AllBackends)
@@ -908,6 +940,28 @@ bool FVirtualizationManager::ShouldDebugDisablePulling(FStringView BackendConfig
 	return false;
 }
 
+bool FVirtualizationManager::ShouldDebugFailPulling() const
+{
+	if (DebugMissChance == 0.0f)
+	{
+		return false;
+	}
+	else
+	{
+		// Could consider adding a lock here, although FRandomStream
+		// is thread safe, many threads hitting it could cause a few
+		// threads to get the same results.
+		// Since this is a debug function and the percent is only a
+		// rough guide, adding a lock is considered overkill. This
+		// should only be done if in the future we decide that we want
+		// more accuracy.
+		static FRandomStream RandomStream(NAME_None);
+
+		const float RandValue = RandomStream.FRand() * 100.0f;
+		return RandValue <= DebugMissChance;
+	}
+}
+
 void FVirtualizationManager::MountBackends(const FConfigFile& ConfigFile)
 {
 	TRACE_CPUPROFILER_EVENT_SCOPE(FVirtualizationManager::MountBackends);
@@ -1111,6 +1165,39 @@ bool FVirtualizationManager::TryPushDataToBackend(IVirtualizationBackend& Backen
 	return bPushResult;
 }
 
+FCompressedBuffer FVirtualizationManager::PullDataFromAllBackends(const FIoHash& Id)
+{
+	if (ShouldDebugFailPulling())
+	{
+		UE_LOG(LogVirtualization, Verbose, TEXT("Debug miss chance (%.1f%%) invoked when pulling payload '%s'"), DebugMissChance, *LexToString(Id));
+		return FCompressedBuffer();
+	}
+
+	for (IVirtualizationBackend* Backend : PullEnabledBackends)
+	{
+		// Skip if pulling has been disabled on this backend for debug purposes
+		if (Backend->IsOperationDebugDisabled(IVirtualizationBackend::EOperations::Pull))
+		{
+			UE_LOG(LogVirtualization, Verbose, TEXT("Pulling from backend '%s' is debug disabled for payload '%s'"), *Backend->GetDebugName(), *LexToString(Id));
+			continue;
+		}
+
+		FCompressedBuffer Payload = PullDataFromBackend(*Backend, Id);
+
+		if (Payload)
+		{
+			if (bEnableCacheAfterPull)
+			{
+				CachePayload(Id, Payload, Backend);
+			}
+
+			return Payload;
+		}
+	}
+
+	return FCompressedBuffer();
+}
+
 FCompressedBuffer FVirtualizationManager::PullDataFromBackend(IVirtualizationBackend& Backend, const FIoHash& Id)
 {
 	COOK_STAT(FCookStats::FScopedStatsCounter Timer(Profiling::GetPullStats(Backend)));
diff --git a/Engine/Source/Developer/Virtualization/Private/VirtualizationManager.h b/Engine/Source/Developer/Virtualization/Private/VirtualizationManager.h
--- a/Engine/Source/Developer/Virtualization/Private/VirtualizationManager.h
+++ b/Engine/Source/Developer/Virtualization/Private/VirtualizationManager.h
@@ -140,8 +140,12 @@ private:
 	void ApplyDebugSettingsFromFromCmdline();
 
 	void OnUpdateMissBackendsFromConsole(const TArray<FString>& Args, FOutputDevice& OutputDevice);
+	void OnUpdateMissChanceFromConsole(const TArray<FString>& Args, FOutputDevice& OutputDevice);
+
 	void UpdateBackendDebugState();
+
 	bool ShouldDebugDisablePulling(FStringView BackendConfigName) const;	
+	bool ShouldDebugFailPulling() const;
 
 	void MountBackends(const FConfigFile& ConfigFile);
 	void ParseHierarchy(const FConfigFile& ConfigFile, const TCHAR* GraphName, const TCHAR* HierarchyKey, const FRegistedFactories& FactoryLookupTable, FBackendArray& PushArray);
@@ -154,6 +158,7 @@ private:
 	bool TryCacheDataToBackend(IVirtualizationBackend& Backend, const FIoHash& Id, const FCompressedBuffer& Payload);
 	bool TryPushDataToBackend(IVirtualizationBackend& Backend, TArrayView<FPushRequest> Requests);
 
+	FCompressedBuffer PullDataFromAllBackends(const FIoHash& Id);
 	FCompressedBuffer PullDataFromBackend(IVirtualizationBackend& Backend, const FIoHash& Id);
 
 	/** 
@@ -245,6 +250,9 @@ private:
 	
 	/** Array of backend names that should have their pull operation disabled */
 	TArray<FString> DebugMissBackends;
+
+	/** The chance that a payload pull can just 'fail' to allow for testing */
+	float DebugMissChance;
 };
 
 } // namespace UE::Virtualization
