Add a new console comment VA.MissCount that will cause the next X number of payload pulls to fail

#rnx

- Moved the debug values to their own structure to make it clearer that they are for debug purposes.
- Moved the code registering the console commands to a specific method.

diff --git a/Engine/Source/Developer/Virtualization/Private/VirtualizationManager.cpp b/Engine/Source/Developer/Virtualization/Private/VirtualizationManager.cpp
--- a/Engine/Source/Developer/Virtualization/Private/VirtualizationManager.cpp
+++ b/Engine/Source/Developer/Virtualization/Private/VirtualizationManager.cpp
@@ -272,22 +272,14 @@ FVirtualizationManager::FVirtualizationManager()
 	// we need to fix the startup/shutdown ordering of Mirage first.
 	COOK_STAT(FCoreDelegates::OnExit.AddStatic(Profiling::LogStats));
 
-	DebugConsoleCommands.Add(IConsoleManager::Get().RegisterConsoleCommand(
-		TEXT("VA.MissBackends"),
-		TEXT("A debug commnad which can be used to disable payload pulling on one or more backends"),
-		FConsoleCommandWithArgsAndOutputDeviceDelegate::CreateRaw(this, &FVirtualizationManager::OnUpdateMissBackendsFromConsole)));
-
-	DebugConsoleCommands.Add(IConsoleManager::Get().RegisterConsoleCommand(
-		TEXT("VA.MissChance"),
-		TEXT("A debug command which can be used to set the chance that a payload pull will fail"),
-		FConsoleCommandWithArgsAndOutputDeviceDelegate::CreateRaw(this, &FVirtualizationManager::OnUpdateMissChanceFromConsole)));
+	RegisterConsoleCommands();
 }
 
 FVirtualizationManager::~FVirtualizationManager()
 {
 	UE_LOG(LogVirtualization, Log, TEXT("Destroying backends"));
 	
-	for (IConsoleCommand* Cmd : DebugConsoleCommands)
+	for (IConsoleCommand* Cmd : DebugValues.ConsoleCommands)
 	{
 		IConsoleManager::Get().UnregisterConsoleObject(Cmd);
 	}
@@ -794,25 +786,43 @@ void FVirtualizationManager::ApplyDebugSettingsFromFromCmdline()
 	FString MissOptions;
 	if (FParse::Value(FCommandLine::Get(), TEXT("-VA-MissBackends="), MissOptions))
 	{
-		MissOptions.ParseIntoArray(DebugMissBackends, TEXT("+"), true);
+		MissOptions.ParseIntoArray(DebugValues.MissBackends, TEXT("+"), true);
 
 		UE_LOG(LogVirtualization, Warning, TEXT("Cmdline has disabled payload pulling for the following backends:"));
-		for (const FString& Backend : DebugMissBackends)
+		for (const FString& Backend : DebugValues.MissBackends)
 		{
 			UE_LOG(LogVirtualization, Warning, TEXT("\t%s"), *Backend);
 		}
 	}
 
-	DebugMissChance = 0.0f;
-	if (FParse::Value(FCommandLine::Get(), TEXT("-VA-MissChance="), DebugMissChance))
+	DebugValues.MissChance = 0.0f;
+	if (FParse::Value(FCommandLine::Get(), TEXT("-VA-MissChance="), DebugValues.MissChance))
 	{
-		DebugMissChance = FMath::Clamp(DebugMissChance, 0.0f, 100.0f);
+		DebugValues.MissChance = FMath::Clamp(DebugValues.MissChance, 0.0f, 100.0f);
 
-		UE_LOG(LogVirtualization, Warning, TEXT("Cmdline has set a %.1f%% chance of a payload pull failing"), DebugMissChance);
+		UE_LOG(LogVirtualization, Warning, TEXT("Cmdline has set a %.1f%% chance of a payload pull failing"), DebugValues.MissChance);
 	}
 }
 
-void FVirtualizationManager::OnUpdateMissBackendsFromConsole(const TArray<FString>& Args, FOutputDevice& OutputDevice)
+void FVirtualizationManager::RegisterConsoleCommands()
+{
+	DebugValues.ConsoleCommands.Add(IConsoleManager::Get().RegisterConsoleCommand(
+		TEXT("VA.MissBackends"),
+		TEXT("A debug commnad which can be used to disable payload pulling on one or more backends"),
+		FConsoleCommandWithArgsAndOutputDeviceDelegate::CreateRaw(this, &FVirtualizationManager::OnUpdateDebugMissBackendsFromConsole)));
+
+	DebugValues.ConsoleCommands.Add(IConsoleManager::Get().RegisterConsoleCommand(
+		TEXT("VA.MissChance"),
+		TEXT("A debug command which can be used to set the chance that a payload pull will fail"),
+		FConsoleCommandWithArgsAndOutputDeviceDelegate::CreateRaw(this, &FVirtualizationManager::OnUpdateDebugMissChanceFromConsole)));
+
+	DebugValues.ConsoleCommands.Add(IConsoleManager::Get().RegisterConsoleCommand(
+		TEXT("VA.MissCount"),
+		TEXT("A debug command which can be used to cause the next X number of payload pulls to fail"),
+		FConsoleCommandWithArgsAndOutputDeviceDelegate::CreateRaw(this, &FVirtualizationManager::OnUpdateDebugMissCountFromConsole)));
+}
+
+void FVirtualizationManager::OnUpdateDebugMissBackendsFromConsole(const TArray<FString>& Args, FOutputDevice& OutputDevice)
 {
 	if (Args.IsEmpty())
 	{
@@ -836,15 +846,15 @@ void FVirtualizationManager::OnUpdateMissBackendsFromConsole(const TArray<FStrin
 	{
 		if (Args[0] == TEXT("reset"))
 		{
-			DebugMissBackends.Empty();
+			DebugValues.MissBackends.Empty();
 			UpdateBackendDebugState();
 		}
 		else if (Args[0] == TEXT("list"))
 		{
-			if (!DebugMissBackends.IsEmpty())
+			if (!DebugValues.MissBackends.IsEmpty())
 			{
 				OutputDevice.Log(TEXT("Disabled backends:"));
-				for (const FString& Backend : DebugMissBackends)
+				for (const FString& Backend : DebugValues.MissBackends)
 				{
 					OutputDevice.Logf(TEXT("\t%s"), *Backend);
 				}
@@ -861,11 +871,11 @@ void FVirtualizationManager::OnUpdateMissBackendsFromConsole(const TArray<FStrin
 	}
 	else if (Args[0] == TEXT("set"))
 	{	
-		DebugMissBackends.Empty(Args.Num() - 1);
+		DebugValues.MissBackends.Empty(Args.Num() - 1);
 
 		for (int32 Index = 1; Index < Args.Num(); ++Index)
 		{
-			DebugMissBackends.Add(Args[Index]);
+			DebugValues.MissBackends.Add(Args[Index]);
 		}
 
 		UpdateBackendDebugState();
@@ -876,7 +886,7 @@ void FVirtualizationManager::OnUpdateMissBackendsFromConsole(const TArray<FStrin
 	}
 }
 
-void FVirtualizationManager::OnUpdateMissChanceFromConsole(const TArray<FString>& Args, FOutputDevice& OutputDevice)
+void FVirtualizationManager::OnUpdateDebugMissChanceFromConsole(const TArray<FString>& Args, FOutputDevice& OutputDevice)
 {
 	if (Args.IsEmpty())
 	{
@@ -889,18 +899,18 @@ void FVirtualizationManager::OnUpdateMissChanceFromConsole(const TArray<FString>
 	}
 	else if (Args.Num() == 1 && Args[0] == TEXT("show"))
 	{
-		OutputDevice.Logf(TEXT("Current debug miss chance: %.1f%%"), DebugMissChance);
+		OutputDevice.Logf(TEXT("Current debug miss chance: %.1f%%"), DebugValues.MissChance);
 	}
 	else if (Args.Num() == 2 && Args[0] == TEXT("set"))
 	{
-		if (::LexTryParseString(DebugMissChance, *Args[1]))
+		if (::LexTryParseString(DebugValues.MissChance, *Args[1]))
 		{
-			DebugMissChance = FMath::Clamp(DebugMissChance, 0.0f, 100.0f);
-			OutputDevice.Logf(TEXT("Current debug miss chance set to %.1f%%"), DebugMissChance);
+			DebugValues.MissChance = FMath::Clamp(DebugValues.MissChance, 0.0f, 100.0f);
+			OutputDevice.Logf(TEXT("Current debug miss chance set to %.1f%%"), DebugValues.MissChance);
 		}
 		else
 		{
-			DebugMissChance = 0.0f;
+			DebugValues.MissChance = 0.0f;
 			OutputDevice.Log(ELogVerbosity::Error, TEXT("Invalid value, current debug miss chance reset to 0.0%"));
 		}
 	}
@@ -910,6 +920,43 @@ void FVirtualizationManager::OnUpdateMissChanceFromConsole(const TArray<FString>
 	}
 }
 
+void FVirtualizationManager::OnUpdateDebugMissCountFromConsole(const TArray<FString>& Args, FOutputDevice& OutputDevice)
+{
+	if (Args.IsEmpty())
+	{
+		OutputDevice.Log(TEXT("VA.MissCount command help"));
+		OutputDevice.Log(TEXT("This command allows you to set the next X number of payload pulls to fail"));
+		OutputDevice.Log(TEXT(""));
+		OutputDevice.Log(TEXT("Commands:"));
+		OutputDevice.Log(TEXT("VA.MissCount show     - prints the current number of future payload pulls that will fail"));
+		OutputDevice.Log(TEXT("VA.MissChance set Num - Sets the number of future payload pulls to fail"));
+	}
+	else if (Args.Num() == 1 && Args[0] == TEXT("show"))
+	{
+		// DebugMissCount could end up negative if many threads are pulling at once, so clamp to 0 as the min value
+		const int32 Value = FMath::Max(DebugValues.MissCount.load(std::memory_order_relaxed), 0);
+		OutputDevice.Logf(TEXT("The next '%d' payload pulls will fail"), Value);
+	}
+	else if (Args.Num() == 2 && Args[0] == TEXT("set"))
+	{
+		int32 ValueToSet = 0;
+		if (::LexTryParseString(ValueToSet, *Args[1]))
+		{
+			DebugValues.MissCount.store(ValueToSet, std::memory_order_relaxed);
+			OutputDevice.Logf(TEXT("The next '%d' payload pulls have been set to fail"), ValueToSet);
+		}
+		else
+		{
+			DebugValues.MissCount.store(0, std::memory_order_relaxed);
+			OutputDevice.Log(ELogVerbosity::Error, TEXT("Invalid value, the number of future payload pulls to fail has been set to zero"));
+		}
+	}
+	else
+	{
+		OutputDevice.Log(ELogVerbosity::Error, TEXT("Invalid args for the VA.MissCount command!"));
+	}
+}
+
 void FVirtualizationManager::UpdateBackendDebugState()
 {
 	for (TUniquePtr<IVirtualizationBackend>& Backend : AllBackends)
@@ -921,17 +968,17 @@ void FVirtualizationManager::UpdateBackendDebugState()
 
 bool FVirtualizationManager::ShouldDebugDisablePulling(FStringView BackendConfigName) const
 {
-	if (DebugMissBackends.IsEmpty())
+	if (DebugValues.MissBackends.IsEmpty())
 	{
 		return false;
 	}
 
-	if (DebugMissBackends[0] == TEXT("All"))
+	if (DebugValues.MissBackends[0] == TEXT("All"))
 	{
 		return true;
 	}
 
-	for (const FString& Name : DebugMissBackends)
+	for (const FString& Name : DebugValues.MissBackends)
 	{
 		if (Name == BackendConfigName)
 		{
@@ -942,9 +989,20 @@ bool FVirtualizationManager::ShouldDebugDisablePulling(FStringView BackendConfig
 	return false;
 }
 
-bool FVirtualizationManager::ShouldDebugFailPulling() const
+bool FVirtualizationManager::ShouldDebugFailPulling()
 {
-	if (DebugMissChance == 0.0f)
+	// We don't want to decrement on every function call to avoid DebugMissCount
+	// underflowing, so we only try to decrement if the count is positive.
+	// It doesn't really matter if the value ends up a little bit negative.
+	if (DebugValues.MissCount.load(std::memory_order_relaxed) > 0)
+	{
+		if (DebugValues.MissCount.fetch_sub(1, std::memory_order_relaxed) > 0)
+		{
+			return true;
+		}
+	}
+
+	if (DebugValues.MissChance == 0.0f)
 	{
 		return false;
 	}
@@ -960,7 +1018,7 @@ bool FVirtualizationManager::ShouldDebugFailPulling() const
 		static FRandomStream RandomStream(NAME_None);
 
 		const float RandValue = RandomStream.FRand() * 100.0f;
-		return RandValue <= DebugMissChance;
+		return RandValue <= DebugValues.MissChance;
 	}
 }
 
@@ -1171,7 +1229,7 @@ FCompressedBuffer FVirtualizationManager::PullDataFromAllBackends(const FIoHash&
 {
 	if (ShouldDebugFailPulling())
 	{
-		UE_LOG(LogVirtualization, Verbose, TEXT("Debug miss chance (%.1f%%) invoked when pulling payload '%s'"), DebugMissChance, *LexToString(Id));
+		UE_LOG(LogVirtualization, Verbose, TEXT("Debug miss chance (%.1f%%) invoked when pulling payload '%s'"), DebugValues.MissChance, *LexToString(Id));
 		return FCompressedBuffer();
 	}
 
diff --git a/Engine/Source/Developer/Virtualization/Private/VirtualizationManager.h b/Engine/Source/Developer/Virtualization/Private/VirtualizationManager.h
--- a/Engine/Source/Developer/Virtualization/Private/VirtualizationManager.h
+++ b/Engine/Source/Developer/Virtualization/Private/VirtualizationManager.h
@@ -139,13 +139,16 @@ private:
 	void ApplyDebugSettingsFromConfigFiles(const FConfigFile& ConfigFile);
 	void ApplyDebugSettingsFromFromCmdline();
 
-	void OnUpdateMissBackendsFromConsole(const TArray<FString>& Args, FOutputDevice& OutputDevice);
-	void OnUpdateMissChanceFromConsole(const TArray<FString>& Args, FOutputDevice& OutputDevice);
+	void RegisterConsoleCommands();
+
+	void OnUpdateDebugMissBackendsFromConsole(const TArray<FString>& Args, FOutputDevice& OutputDevice);
+	void OnUpdateDebugMissChanceFromConsole(const TArray<FString>& Args, FOutputDevice& OutputDevice);
+	void OnUpdateDebugMissCountFromConsole(const TArray<FString>& Args, FOutputDevice& OutputDevice);
 
 	void UpdateBackendDebugState();
 
 	bool ShouldDebugDisablePulling(FStringView BackendConfigName) const;	
-	bool ShouldDebugFailPulling() const;
+	bool ShouldDebugFailPulling();
 
 	void MountBackends(const FConfigFile& ConfigFile);
 	void ParseHierarchy(const FConfigFile& ConfigFile, const TCHAR* GraphName, const TCHAR* HierarchyKey, const FRegistedFactories& FactoryLookupTable, FBackendArray& PushArray);
@@ -245,14 +248,20 @@ private:
 
 	// Members after this point at used for debugging operations only!
 
-	/** Console commands that the manager has registered */
-	TArray<IConsoleCommand*> DebugConsoleCommands;
-	
-	/** Array of backend names that should have their pull operation disabled */
-	TArray<FString> DebugMissBackends;
+	struct FDebugValues
+	{
+		/** Console commands that the manager has registered */
+		TArray<IConsoleCommand*> ConsoleCommands;
+
+		/** Array of backend names that should have their pull operation disabled */
+		TArray<FString> MissBackends;
+
+		/** The chance that a payload pull can just 'fail' to allow for testing */
+		float MissChance;
 
-	/** The chance that a payload pull can just 'fail' to allow for testing */
-	float DebugMissChance;
+		/** The number of upcoming payload pulls that should be failed */
+		std::atomic<int32> MissCount = 0;
+	} DebugValues;
 };
 
 } // namespace UE::Virtualization
