The debug setting to run virtualization in single threaded mode is now access from the commandline via '-VA-SingleThreaded' or the console via 'VA.SingleThreaded' and no longer by the config file system.

#rnx

- It was probably a bad idea to have debug settings like this applied from the config file as it would be fairly easy for people to accidently forget about them or submit them.
- The single threaded command can now be set via the commandline on start up using the '-VA-XXX' synctax that the system has started to adopt.
- Additionally it has been exposed as a console command so that it can be toggles on and off at runtime.
- Removed 'ForceSingleThreaded' from BaseEngine.ini and removed the parsing code for this from FVirtualizationManager::ApplyDebugSettingsFromConfigFiles.
- Removed 'FailPayloadPullOperations' from BaseEngine.ini, it hasn't been a valid setting for a while.
- Changed ConsoleCommands to ConsoleObjects and use it to store the base class IConsoleObject instead so it can contain console variables
- When pushing payloads we now only take the singlethreaded lock after checking if we have any backends/are enabled.

diff --git a/Engine/Config/BaseEngine.ini b/Engine/Config/BaseEngine.ini
--- a/Engine/Config/BaseEngine.ini
+++ b/Engine/Config/BaseEngine.ini
@@ -1203,8 +1203,6 @@ FilterEnginePluginContent=True
 
 ; Options to help debug/profile problems with the content virtualization system
 [Core.ContentVirtualizationDebugOptions]
-ForceSingleThreaded=False
-FailPayloadPullOperations=False
 ValidateAfterPushOperation=False
 
 ; A backend graph with no backends
diff --git a/Engine/Source/Developer/Virtualization/Private/VirtualizationManager.cpp b/Engine/Source/Developer/Virtualization/Private/VirtualizationManager.cpp
--- a/Engine/Source/Developer/Virtualization/Private/VirtualizationManager.cpp
+++ b/Engine/Source/Developer/Virtualization/Private/VirtualizationManager.cpp
@@ -262,7 +262,6 @@ FVirtualizationManager::FVirtualizationManager()
 	, bEnableCacheAfterPull(true)
 	, MinPayloadLength(0)
 	, BackendGraphName(TEXT("ContentVirtualizationBackendGraph_None"))
-	, bForceSingleThreaded(false)
 	, bValidateAfterPushOperation(false)
 {
 	TRACE_CPUPROFILER_EVENT_SCOPE(FVirtualizationManager::FVirtualizationManager);
@@ -271,17 +270,15 @@ FVirtualizationManager::FVirtualizationManager()
 	// TODO: We should just be able to call the logging in the destructor, but 
 	// we need to fix the startup/shutdown ordering of Mirage first.
 	COOK_STAT(FCoreDelegates::OnExit.AddStatic(Profiling::LogStats));
-
-	RegisterConsoleCommands();
 }
 
 FVirtualizationManager::~FVirtualizationManager()
 {
 	UE_LOG(LogVirtualization, Log, TEXT("Destroying backends"));
-	
-	for (IConsoleCommand* Cmd : DebugValues.ConsoleCommands)
+
+	for (IConsoleObject* ConsoleObject : DebugValues.ConsoleObjects)
 	{
-		IConsoleManager::Get().UnregisterConsoleObject(Cmd);
+		IConsoleManager::Get().UnregisterConsoleObject(ConsoleObject);
 	}
 
 	LocalCachableBackends.Empty();
@@ -305,6 +302,10 @@ bool FVirtualizationManager::Initialize(const FInitParams& InitParams)
 	ApplySettingsFromCmdline();
 	ApplyDebugSettingsFromFromCmdline();
 
+	// Do this after all of the command line settings have been processed and any 
+	// requested debug value changes already set.
+	RegisterConsoleCommands();
+
 	MountBackends(InitParams.ConfigFile);
 
 	return true;
@@ -404,14 +405,14 @@ bool FVirtualizationManager::PushData(TArrayView<FPushRequest> Requests, EStorag
 		return true;
 	}
 
-	FConditionalScopeLock _(&ForceSingleThreadedCS, bForceSingleThreaded);
-
 	// Early out if there are no backends
 	if (!IsEnabled() || bEnablePayloadPushing == false)
 	{
 		return false;
 	}
 
+	FConditionalScopeLock _(&ForceSingleThreadedCS, DebugValues.bSingleThreaded);
+
 	// TODO: Note that all push operations are currently synchronous, probably 
 	// should change to async at some point, although this makes handling failed
 	// pushed much more difficult.
@@ -482,7 +483,7 @@ FCompressedBuffer FVirtualizationManager::PullData(const FIoHash& Id)
 		return FCompressedBuffer();
 	}
 
-	FConditionalScopeLock _(&ForceSingleThreadedCS, bForceSingleThreaded);
+	FConditionalScopeLock _(&ForceSingleThreadedCS, DebugValues.bSingleThreaded);
 
 	GetNotificationEvent().Broadcast(IVirtualizationSystem::PullBegunNotification, Id);
 
@@ -534,7 +535,7 @@ EQueryResult FVirtualizationManager::QueryPayloadStatuses(TArrayView<const FIoHa
 	Results.SetNum(Ids.Num());
 
 	{
-		FConditionalScopeLock _(&ForceSingleThreadedCS, bForceSingleThreaded);
+		FConditionalScopeLock _(&ForceSingleThreadedCS, DebugValues.bSingleThreaded);
 
 		for (IVirtualizationBackend* Backend : Backends)
 		{
@@ -748,12 +749,6 @@ void FVirtualizationManager::ApplySettingsFromCmdline()
 		UE_LOG(LogVirtualization, Display, TEXT("Backend graph overriden from the cmdline: '%s'"), *CmdlineGraphName);
 		BackendGraphName = CmdlineGraphName;
 	}
-
-	if (FParse::Param(FCommandLine::Get(), TEXT("VirtualizationForceSingleThreaded")))
-	{
-		bForceSingleThreaded = true;
-		UE_LOG(LogVirtualization, Display, TEXT("ForceSingleThreaded overriden from the cmdline: true"));
-	}
 }
 
 void FVirtualizationManager::ApplyDebugSettingsFromConfigFiles(const FConfigFile& ConfigFile)
@@ -761,13 +756,6 @@ void FVirtualizationManager::ApplyDebugSettingsFromConfigFiles(const FConfigFile
 	UE_LOG(LogVirtualization, Display, TEXT("Loading virtualization manager debugging settings from config files..."));
 
 	// Note that the debug settings are optional and could be left out of the config files entirely
-	bool bForceSingleThreadedFromIni = false;
-	if (ConfigFile.GetBool(TEXT("Core.ContentVirtualizationDebugOptions"), TEXT("ForceSingleThreaded"), bForceSingleThreadedFromIni))
-	{
-		bForceSingleThreaded = bForceSingleThreadedFromIni;
-		UE_LOG(LogVirtualization, Display, TEXT("\tForceSingleThreaded : %s"), bForceSingleThreaded ? TEXT("true") : TEXT("false"));
-	}
-
 	bool bValidateAfterPushOperationFromIni = false;
 	if (ConfigFile.GetBool(TEXT("Core.ContentVirtualizationDebugOptions"), TEXT("ValidateAfterPushOperation"), bValidateAfterPushOperationFromIni))
 	{
@@ -777,12 +765,17 @@ void FVirtualizationManager::ApplyDebugSettingsFromConfigFiles(const FConfigFile
 
 	// Some debug options will cause intentional breaks or slow downs for testing purposes, if these are enabled then we should give warning/errors 
 	// so it is clear in the log that future failures are being caused by the given dev option.
-	UE_CLOG(bForceSingleThreaded, LogVirtualization, Warning, TEXT("ForceSingleThreaded is enabled, virtualization will run in single threaded mode and may be slower!"));
 	UE_CLOG(bValidateAfterPushOperation, LogVirtualization, Error, TEXT("ValidateAfterPushOperation is enabled, each push will be followed by a pull to validate it!"));
 }
 
 void FVirtualizationManager::ApplyDebugSettingsFromFromCmdline()
 {
+	if (FParse::Param(FCommandLine::Get(), TEXT("VA-SingleThreaded")))
+	{
+		DebugValues.bSingleThreaded = true;
+		UE_LOG(LogVirtualization, Warning, TEXT("Cmdline has set the virtualization system to run single threaded"));
+	}
+
 	FString MissOptions;
 	if (FParse::Value(FCommandLine::Get(), TEXT("-VA-MissBackends="), MissOptions))
 	{
@@ -806,20 +799,29 @@ void FVirtualizationManager::ApplyDebugSettingsFromFromCmdline()
 
 void FVirtualizationManager::RegisterConsoleCommands()
 {
-	DebugValues.ConsoleCommands.Add(IConsoleManager::Get().RegisterConsoleCommand(
+	DebugValues.ConsoleObjects.Add(IConsoleManager::Get().RegisterConsoleCommand(
 		TEXT("VA.MissBackends"),
 		TEXT("A debug commnad which can be used to disable payload pulling on one or more backends"),
-		FConsoleCommandWithArgsAndOutputDeviceDelegate::CreateRaw(this, &FVirtualizationManager::OnUpdateDebugMissBackendsFromConsole)));
+		FConsoleCommandWithArgsAndOutputDeviceDelegate::CreateRaw(this, &FVirtualizationManager::OnUpdateDebugMissBackendsFromConsole)
+	));
 
-	DebugValues.ConsoleCommands.Add(IConsoleManager::Get().RegisterConsoleCommand(
+	DebugValues.ConsoleObjects.Add(IConsoleManager::Get().RegisterConsoleCommand(
 		TEXT("VA.MissChance"),
 		TEXT("A debug command which can be used to set the chance that a payload pull will fail"),
-		FConsoleCommandWithArgsAndOutputDeviceDelegate::CreateRaw(this, &FVirtualizationManager::OnUpdateDebugMissChanceFromConsole)));
+		FConsoleCommandWithArgsAndOutputDeviceDelegate::CreateRaw(this, &FVirtualizationManager::OnUpdateDebugMissChanceFromConsole)
+	));
 
-	DebugValues.ConsoleCommands.Add(IConsoleManager::Get().RegisterConsoleCommand(
+	DebugValues.ConsoleObjects.Add(IConsoleManager::Get().RegisterConsoleCommand(
 		TEXT("VA.MissCount"),
 		TEXT("A debug command which can be used to cause the next X number of payload pulls to fail"),
-		FConsoleCommandWithArgsAndOutputDeviceDelegate::CreateRaw(this, &FVirtualizationManager::OnUpdateDebugMissCountFromConsole)));
+		FConsoleCommandWithArgsAndOutputDeviceDelegate::CreateRaw(this, &FVirtualizationManager::OnUpdateDebugMissCountFromConsole)
+	));
+
+	DebugValues.ConsoleObjects.Add(IConsoleManager::Get().RegisterConsoleVariableRef(
+		TEXT("VA.SingleThreaded"),
+		DebugValues.bSingleThreaded,
+		TEXT("When set the asset virtualization system will only access backends in a single threaded manner")
+	));
 }
 
 void FVirtualizationManager::OnUpdateDebugMissBackendsFromConsole(const TArray<FString>& Args, FOutputDevice& OutputDevice)
diff --git a/Engine/Source/Developer/Virtualization/Private/VirtualizationManager.h b/Engine/Source/Developer/Virtualization/Private/VirtualizationManager.h
--- a/Engine/Source/Developer/Virtualization/Private/VirtualizationManager.h
+++ b/Engine/Source/Developer/Virtualization/Private/VirtualizationManager.h
@@ -9,7 +9,7 @@
 
 #include "Virtualization/VirtualizationSystem.h"
 
-struct IConsoleCommand;
+class IConsoleObject;
 class FOutputDevice;
 
 /**
@@ -211,9 +211,6 @@ private:
 	
 	/** Should payloads in engine plugin content packages before filtered out and never virtualized */
 	bool bFilterEnginePluginContent;
-
-	/** Debugging option: When enabled all public operations will be performed as single threaded. This is intended to aid debugging and not for production use.*/
-	bool bForceSingleThreaded;
 	
 	/**
 	 * Debugging option: When enabled we will immediately 'pull' each payload after it has been 'pushed' and compare it to the original payload source to make 
@@ -250,8 +247,11 @@ private:
 
 	struct FDebugValues
 	{
-		/** Console commands that the manager has registered */
-		TArray<IConsoleCommand*> ConsoleCommands;
+		/** All of the console commands/variables that we register, so they can be unregistered when the manager is destroyed */
+		TArray<IConsoleObject*> ConsoleObjects;
+
+		/** When enabled all public operations will be performed as single threaded */
+		bool bSingleThreaded = false;
 
 		/** Array of backend names that should have their pull operation disabled */
 		TArray<FString> MissBackends;
