The debug option for validating payloads after pushing them to storage is now a commandline & console option rather than a config file option.

#rnx

- Remove the Core.ContentVirtualizationDebugOptions section from the ini files entirely.
- Remove the ::ApplyDebugSettingsFromConfigFiles method as we no longer load any debug values from the config files.

diff --git a/Engine/Config/BaseEngine.ini b/Engine/Config/BaseEngine.ini
--- a/Engine/Config/BaseEngine.ini
+++ b/Engine/Config/BaseEngine.ini
@@ -1201,10 +1201,6 @@ FilterEnginePluginContent=True
 [/Script/Virtualization.VirtualizationFilterSettings]
 !ExcludePackagePaths=ClearArray
 
-; Options to help debug/profile problems with the content virtualization system
-[Core.ContentVirtualizationDebugOptions]
-ValidateAfterPushOperation=False
-
 ; A backend graph with no backends
 [ContentVirtualizationBackendGraph_None]
 LocalStorageHierarchy=()
diff --git a/Engine/Source/Developer/Virtualization/Private/VirtualizationManager.cpp b/Engine/Source/Developer/Virtualization/Private/VirtualizationManager.cpp
--- a/Engine/Source/Developer/Virtualization/Private/VirtualizationManager.cpp
+++ b/Engine/Source/Developer/Virtualization/Private/VirtualizationManager.cpp
@@ -262,7 +262,6 @@ FVirtualizationManager::FVirtualizationManager()
 	, bEnableCacheAfterPull(true)
 	, MinPayloadLength(0)
 	, BackendGraphName(TEXT("ContentVirtualizationBackendGraph_None"))
-	, bValidateAfterPushOperation(false)
 {
 	TRACE_CPUPROFILER_EVENT_SCOPE(FVirtualizationManager::FVirtualizationManager);
 
@@ -297,7 +296,6 @@ bool FVirtualizationManager::Initialize(const FInitParams& InitParams)
 	ProjectName = InitParams.ProjectName;
 
 	ApplySettingsFromConfigFiles(InitParams.ConfigFile);
-	ApplyDebugSettingsFromConfigFiles(InitParams.ConfigFile);
 	
 	ApplySettingsFromCmdline();
 	ApplyDebugSettingsFromFromCmdline();
@@ -433,9 +431,8 @@ bool FVirtualizationManager::PushData(TArrayView<FPushRequest> Requests, EStorag
 			ErrorCount++;
 		}
 		
-		// Debugging operation where we immediately try to pull the payload after each push (when possible) and assert 
-		// that the pulled payload is the same as the original
-		if (bValidateAfterPushOperation && bResult == true && Backend->IsOperationSupported(IVirtualizationBackend::EOperations::Pull))
+		// Debug operation to validate that the payload we just pushed can be retrieved from storage
+		if (DebugValues.bValidateAfterPush && bResult == true && Backend->IsOperationSupported(IVirtualizationBackend::EOperations::Pull))
 		{
 			for (FPushRequest& Request : ValidatedRequests)
 			{
@@ -751,23 +748,6 @@ void FVirtualizationManager::ApplySettingsFromCmdline()
 	}
 }
 
-void FVirtualizationManager::ApplyDebugSettingsFromConfigFiles(const FConfigFile& ConfigFile)
-{
-	UE_LOG(LogVirtualization, Display, TEXT("Loading virtualization manager debugging settings from config files..."));
-
-	// Note that the debug settings are optional and could be left out of the config files entirely
-	bool bValidateAfterPushOperationFromIni = false;
-	if (ConfigFile.GetBool(TEXT("Core.ContentVirtualizationDebugOptions"), TEXT("ValidateAfterPushOperation"), bValidateAfterPushOperationFromIni))
-	{
-		bValidateAfterPushOperation = bValidateAfterPushOperationFromIni;
-		UE_LOG(LogVirtualization, Display, TEXT("\tValidateAfterPushOperation : %s"), bValidateAfterPushOperation ? TEXT("true") : TEXT("false"));
-	}
-
-	// Some debug options will cause intentional breaks or slow downs for testing purposes, if these are enabled then we should give warning/errors 
-	// so it is clear in the log that future failures are being caused by the given dev option.
-	UE_CLOG(bValidateAfterPushOperation, LogVirtualization, Error, TEXT("ValidateAfterPushOperation is enabled, each push will be followed by a pull to validate it!"));
-}
-
 void FVirtualizationManager::ApplyDebugSettingsFromFromCmdline()
 {
 	if (FParse::Param(FCommandLine::Get(), TEXT("VA-SingleThreaded")))
@@ -776,6 +756,12 @@ void FVirtualizationManager::ApplyDebugSettingsFromFromCmdline()
 		UE_LOG(LogVirtualization, Warning, TEXT("Cmdline has set the virtualization system to run single threaded"));
 	}
 
+	if (FParse::Param(FCommandLine::Get(), TEXT("VA-ValidatePushes")))
+	{
+		DebugValues.bValidateAfterPush = true;
+		UE_LOG(LogVirtualization, Warning, TEXT("Cmdline has set the virtualization system to pull each payload after pushing to either local or persistent storage"));
+	}
+
 	FString MissOptions;
 	if (FParse::Value(FCommandLine::Get(), TEXT("-VA-MissBackends="), MissOptions))
 	{
@@ -822,6 +808,12 @@ void FVirtualizationManager::RegisterConsoleCommands()
 		DebugValues.bSingleThreaded,
 		TEXT("When set the asset virtualization system will only access backends in a single threaded manner")
 	));
+
+	DebugValues.ConsoleObjects.Add(IConsoleManager::Get().RegisterConsoleVariableRef(
+		TEXT("VA.ValidatePushes"),
+		DebugValues.bValidateAfterPush,
+		TEXT("When set the asset virtualization system will pull each payload after pushing to either local or persistent storage")
+	));
 }
 
 void FVirtualizationManager::OnUpdateDebugMissBackendsFromConsole(const TArray<FString>& Args, FOutputDevice& OutputDevice)
@@ -1171,9 +1163,8 @@ void FVirtualizationManager::CachePayload(const FIoHash& Id, const FCompressedBu
 					*LexToString(Id),
 					*BackendToCache->GetDebugName());
 
-		// Debugging operation where we immediately try to pull the payload after each push (when possible) and assert 
-		// that the pulled payload is the same as the original
-		if (bValidateAfterPushOperation && bResult && BackendToCache->IsOperationSupported(IVirtualizationBackend::EOperations::Pull))
+		// Debug operation to validate that the payload we just cached can be retrieved from storage
+		if (DebugValues.bValidateAfterPush && bResult && BackendToCache->IsOperationSupported(IVirtualizationBackend::EOperations::Pull))
 		{
 			FCompressedBuffer PulledPayload = PullDataFromBackend(*BackendToCache, Id);
 			checkf(	Payload.GetRawHash() == PulledPayload.GetRawHash(), 
diff --git a/Engine/Source/Developer/Virtualization/Private/VirtualizationManager.h b/Engine/Source/Developer/Virtualization/Private/VirtualizationManager.h
--- a/Engine/Source/Developer/Virtualization/Private/VirtualizationManager.h
+++ b/Engine/Source/Developer/Virtualization/Private/VirtualizationManager.h
@@ -136,7 +136,6 @@ private:
 	void ApplySettingsFromConfigFiles(const FConfigFile& ConfigFile);
 	void ApplySettingsFromCmdline();
 	
-	void ApplyDebugSettingsFromConfigFiles(const FConfigFile& ConfigFile);
 	void ApplyDebugSettingsFromFromCmdline();
 
 	void RegisterConsoleCommands();
@@ -212,13 +211,6 @@ private:
 	/** Should payloads in engine plugin content packages before filtered out and never virtualized */
 	bool bFilterEnginePluginContent;
 	
-	/**
-	 * Debugging option: When enabled we will immediately 'pull' each payload after it has been 'pushed' and compare it to the original payload source to make 
-	 * sure that it can be pulled correctly.
-	 * This is intended to aid debugging and not for production use.
-	 */
-	bool bValidateAfterPushOperation;
-
 	/** The name of the current project */
 	FString ProjectName;
 
@@ -253,6 +245,12 @@ private:
 		/** When enabled all public operations will be performed as single threaded */
 		bool bSingleThreaded = false;
 
+		/** 
+		 * When enabled we will immediately 'pull' each payload after it has been pushed to either local or persistent
+		 * storage and compare it to the original payload source to make sure that it was uploaded correctly 
+		 */
+		bool bValidateAfterPush = false;
+
 		/** Array of backend names that should have their pull operation disabled */
 		TArray<FString> MissBackends;
 
