Renamed the virtualization cmdline '-BackendGraph' to '-VA-BackendGraph' to match the other commandlines.

#rnx

- Moving this to use the format that the other command lines do will make it easier to find.
- Moved the parsing code from ::ApplySettingsFromCmdline to ::ApplyDebugSettingsFromFromCmdline to indicate that it is a debug command and not expected for production use.
- Removed ::ApplySettingsFromCmdline as we no longer have non-debug commandlines.
- There is no console command version because we cannot change the graph mid process.

diff --git a/Engine/Source/Developer/Virtualization/Private/VirtualizationManager.cpp b/Engine/Source/Developer/Virtualization/Private/VirtualizationManager.cpp
--- a/Engine/Source/Developer/Virtualization/Private/VirtualizationManager.cpp
+++ b/Engine/Source/Developer/Virtualization/Private/VirtualizationManager.cpp
@@ -296,8 +296,6 @@ bool FVirtualizationManager::Initialize(const FInitParams& InitParams)
 	ProjectName = InitParams.ProjectName;
 
 	ApplySettingsFromConfigFiles(InitParams.ConfigFile);
-	
-	ApplySettingsFromCmdline();
 	ApplyDebugSettingsFromFromCmdline();
 
 	// Do this after all of the command line settings have been processed and any 
@@ -738,16 +736,6 @@ void FVirtualizationManager::ApplySettingsFromConfigFiles(const FConfigFile& Con
 	}	
 }
 
-void FVirtualizationManager::ApplySettingsFromCmdline()
-{
-	FString CmdlineGraphName;
-	if (FParse::Value(FCommandLine::Get(), TEXT("-BackendGraph="), CmdlineGraphName))
-	{
-		UE_LOG(LogVirtualization, Display, TEXT("Backend graph overriden from the cmdline: '%s'"), *CmdlineGraphName);
-		BackendGraphName = CmdlineGraphName;
-	}
-}
-
 void FVirtualizationManager::ApplyDebugSettingsFromFromCmdline()
 {
 	if (FParse::Param(FCommandLine::Get(), TEXT("VA-SingleThreaded")))
@@ -762,6 +750,13 @@ void FVirtualizationManager::ApplyDebugSettingsFromFromCmdline()
 		UE_LOG(LogVirtualization, Warning, TEXT("Cmdline has set the virtualization system to pull each payload after pushing to either local or persistent storage"));
 	}
 
+	FString CmdlineGraphName;
+	if (FParse::Value(FCommandLine::Get(), TEXT("-VA-BackendGraph="), CmdlineGraphName))
+	{
+		UE_LOG(LogVirtualization, Display, TEXT("Backend graph overriden from the cmdline: '%s'"), *CmdlineGraphName);
+		BackendGraphName = CmdlineGraphName;
+	}
+
 	FString MissOptions;
 	if (FParse::Value(FCommandLine::Get(), TEXT("-VA-MissBackends="), MissOptions))
 	{
diff --git a/Engine/Source/Developer/Virtualization/Private/VirtualizationManager.h b/Engine/Source/Developer/Virtualization/Private/VirtualizationManager.h
--- a/Engine/Source/Developer/Virtualization/Private/VirtualizationManager.h
+++ b/Engine/Source/Developer/Virtualization/Private/VirtualizationManager.h
@@ -134,8 +134,6 @@ private:
 private:
 
 	void ApplySettingsFromConfigFiles(const FConfigFile& ConfigFile);
-	void ApplySettingsFromCmdline();
-	
 	void ApplyDebugSettingsFromFromCmdline();
 
 	void RegisterConsoleCommands();
