The source control backend can now limit the number of threads that can make concurrent connections.

#rnx

- In the worst case scenario (where a machine can only access the source control backend) we will want to limit the potential number of concurrent connections that they can make.
- In testing, not limiting the connections could cause the occasional connection failure and I imagine having a full team running in the worse case scenario would quickly overwhelm the server.
- Currently we default to 8 connections (the limit allowed in UGS)
- The semaphore is a quick and dirty implementation similar to std::counting_semaphore. I am not spending much time on it as the worst case scenario is not expected to be common and if it does occur, the system is going to be painfully slow no matter how good the thread throttling.

### Future work
- This work could be  taken a lot further by either
-- a) reusing connections rather than creating one for each pull, which will require changes to our source control api
-- b) gathering requests from many requesting threads once the limit has been reached and then sending the requests in batches
--c) reserve one connection for the game thread to prevent the user waiting?

diff --git a/Engine/Source/Developer/Virtualization/Private/VirtualizationSourceControlBackend.cpp b/Engine/Source/Developer/Virtualization/Private/VirtualizationSourceControlBackend.cpp
--- a/Engine/Source/Developer/Virtualization/Private/VirtualizationSourceControlBackend.cpp
+++ b/Engine/Source/Developer/Virtualization/Private/VirtualizationSourceControlBackend.cpp
@@ -27,6 +27,114 @@
 namespace UE::Virtualization
 {
 
+/** 
+ * A quick and dirty, poor mans implementation of a std::counting_semaphore that we can use to 
+ * limit the number of threads that can create a new perforce connection when pulling or 
+ * pushing payloads.
+ * 
+ * In the worst case scenario where a user needs to pull all of their payloads from the source
+ * control backend rather than a faster backend we need to make sure that they will not overwhelm
+ * their server with requests.
+ * In the future we can use this sort of limit to help gather requests from many threads into a 
+ * single batch request from the server which will be much more efficient than the current 'one
+ * payload, one request' system. Although we might want to consider gathering multiple requests
+ * at a higher level so that all backends can work on the same batching principle.
+ */
+class FSemaphore 
+{
+public:
+	UE_NONCOPYABLE(FSemaphore);
+
+	enum class EAcquireResult
+	{
+		/** The acquire was a success and the thread can continue */
+		Success,
+		/** The wait event failed, the semaphore object is no longer safe */
+		EventFailed
+	};
+
+	FSemaphore() = delete;
+	explicit FSemaphore(int32 InitialCount)
+		: WaitEvent(EEventMode::ManualReset)
+		, Counter(InitialCount)
+		, DebugCount(0)
+	{
+
+	}
+
+	~FSemaphore()
+	{
+		checkf(DebugCount == 0, TEXT("'%d' threads are still waiting on the UE::Virtualization::FSemaphore being destroyed"));
+	}
+
+	/** Will block until the calling thread can pass through the semaphore. Note that it might return an error if the WaitEvent fails */
+	EAcquireResult Acquire()
+	{
+		CriticalSection.Lock();
+		DebugCount++;
+
+		while (Counter-- <= 0)
+		{
+			Counter++;
+			WaitEvent->Reset();
+
+			CriticalSection.Unlock();
+
+			if (!WaitEvent->Wait())
+			{
+				--DebugCount;
+				return EAcquireResult::EventFailed;
+			}
+			CriticalSection.Lock();
+		}
+
+		CriticalSection.Unlock();
+
+		return EAcquireResult::Success;
+	}
+
+	void Release()
+	{
+		FScopeLock _(&CriticalSection);
+
+		Counter++;
+		WaitEvent->Trigger();
+
+		--DebugCount;
+	}
+
+private:
+	FEventRef WaitEvent;
+	FCriticalSection CriticalSection;
+
+	std::atomic<int32> Counter;
+	std::atomic<int32> DebugCount;
+};
+
+/** Structure to make it easy to acquire/release a FSemaphore for a given scope */
+struct FSemaphoreScopeLock
+{
+	FSemaphoreScopeLock(FSemaphore* InSemaphore)
+		: Semaphore(InSemaphore)
+	{
+		if (Semaphore != nullptr)
+		{
+			Semaphore->Acquire();
+		}
+	}
+
+	~FSemaphoreScopeLock()
+	{
+		if (Semaphore != nullptr)
+		{
+			Semaphore->Release();
+		}
+	}
+
+private:
+	FSemaphore* Semaphore;
+};
+
 /** Utility function to create a directory to submit payloads from. */
 [[nodiscard]] static bool TryCreateSubmissionSessionDirectory(FStringView SessionDirectoryPath)
 {
@@ -213,6 +321,11 @@ FCompressedBuffer FSourceControlBackend::PullData(const FIoHash& Id)
 	TStringBuilder<512> DepotPath;
 	CreateDepotPath(Id, DepotPath);
 
+	// TODO: When multiple threads are blocked waiting on this we could gather X payloads together and make a single
+	// batch request on the same connection, which should be a lot faster with less overhead.
+	// Although ideally this backend will not get hit very often.
+	FSemaphoreScopeLock _(ConcurrentConnectionLimit.Get());
+
 	UE_LOG(LogVirtualization, Verbose, TEXT("[%s] Attempting to pull '%s' from source control"), *GetDebugName(), *DepotPath);
 
 	int32 Retries = 0;
@@ -369,6 +482,8 @@ bool FSourceControlBackend::PushData(TArrayView<FPushRequest> Requests)
 
 	check(Requests.Num() == FilesToSubmit.Num());
 
+	FSemaphoreScopeLock _(ConcurrentConnectionLimit.Get());
+
 	TStringBuilder<64> WorkspaceName;
 	WorkspaceName << TEXT("VASubmission-") << SessionGuid;
 
@@ -571,24 +686,46 @@ bool FSourceControlBackend::TryApplySettingsFromConfigFiles(const FString& Confi
 		return false;
 	}
 
-	// Optional config values
-	FParse::Bool(*ConfigEntry, TEXT("UsePartitionedClient="), bUsePartitionedClient);
+	// Now parse the optional config values
 
-	int32 RetryCountIniFile = INDEX_NONE;
-	if (FParse::Value(*ConfigEntry, TEXT("RetryCount="), RetryCountIniFile))
+	// Check to see if we should use partitioned clients or not. This is a perforce specific optimization to make the workspace churn cheaper on the server
 	{
-		RetryCount = RetryCountIniFile;
+		FParse::Bool(*ConfigEntry, TEXT("UsePartitionedClient="), bUsePartitionedClient);
+		UE_LOG(LogVirtualization, Log, TEXT("[%s] Using partitioned clients: '%s'"), *GetDebugName(), bUsePartitionedClient ? TEXT("true") : TEXT("false"));
 	}
 
-	int32 RetryWaitTimeMSIniFile = INDEX_NONE;
-	if (FParse::Value(*ConfigEntry, TEXT("RetryWaitTime="), RetryWaitTimeMSIniFile))
+	// Allow the source control backend to retry failed pulls
 	{
-		RetryWaitTimeMS = RetryWaitTimeMSIniFile;
+		int32 RetryCountIniFile = INDEX_NONE;
+		if (FParse::Value(*ConfigEntry, TEXT("RetryCount="), RetryCountIniFile))
+		{
+			RetryCount = RetryCountIniFile;
+		}
+
+		int32 RetryWaitTimeMSIniFile = INDEX_NONE;
+		if (FParse::Value(*ConfigEntry, TEXT("RetryWaitTime="), RetryWaitTimeMSIniFile))
+		{
+			RetryWaitTimeMS = RetryWaitTimeMSIniFile;
+		}
+
+		UE_LOG(LogVirtualization, Log, TEXT("[%s] Will retry failed downloads attempts %d time(s) with a gap of %dms betwen them"), *GetDebugName(), RetryCount, RetryWaitTimeMS);
 	}
 
-	// Now log a summary of the optional settings to make issues easier to diagnose
-	UE_LOG(LogVirtualization, Log, TEXT("[%s] Using partitioned clients: '%s'"), *GetDebugName(), bUsePartitionedClient ? TEXT("true") : TEXT("false"));
-	UE_LOG(LogVirtualization, Log, TEXT("[%s] Will retry failed downloads attempts %d time(s) with a gap of %dms betwen them"), *GetDebugName(), RetryCount, RetryWaitTimeMS);
+	// Allow the number of concurrent connections to be limited
+	{
+		int32 MaxLimit = 8; // We use the UGS max of 8 as the default
+		FParse::Value(*ConfigEntry, TEXT("MaxConnections="), MaxLimit);
+
+		if (MaxLimit != INDEX_NONE)
+		{
+			ConcurrentConnectionLimit = MakeUnique<FSemaphore>(MaxLimit);
+			UE_LOG(LogVirtualization, Log, TEXT("[%s] Limted to %d concurrent source control connections"), *GetDebugName(), MaxLimit);
+		}
+		else
+		{
+			UE_LOG(LogVirtualization, Log, TEXT("[%s] Has no limit to it's concurrent source control connections"), *GetDebugName());
+		}
+	}	
 
 	if (!FindSubmissionWorkingDir(ConfigEntry))
 	{
diff --git a/Engine/Source/Developer/Virtualization/Private/VirtualizationSourceControlBackend.h b/Engine/Source/Developer/Virtualization/Private/VirtualizationSourceControlBackend.h
--- a/Engine/Source/Developer/Virtualization/Private/VirtualizationSourceControlBackend.h
+++ b/Engine/Source/Developer/Virtualization/Private/VirtualizationSourceControlBackend.h
@@ -11,6 +11,8 @@ class ISourceControlProvider;
 namespace UE::Virtualization
 {
 
+class FSemaphore;
+
 /**
  * This backend can be used to access payloads stored in source control.
  * The backend doesn't 'check out' a payload file but instead will just download the payload as
@@ -91,6 +93,9 @@ private:
 	/** Should we try to make the temp client partitioned or not? */
 	bool bUsePartitionedClient = true;
 
+	/** A counted semaphore that will limit the number of concurrent connections that we can make */
+	TUniquePtr<UE::Virtualization::FSemaphore> ConcurrentConnectionLimit;
+	
 	/** The number of times to retry pulling a payload from the depot */
 	int32 RetryCount = 2;
 
