When virtualizing payloads we now load the payloads from disk on demand rather than all at once. This allows us to avoid memory spikes when virtualizing large amounts of data.

#rnx

- In the old code when submitting packages for virtualization we would load all local payloads into memory then pass them to the virtualization system to be pushed to persistent storage. This works fine for an average users submits but when resaving large projects we can have submits with thousands of packages containing gigabytes of payloads.
- FPushRequest has been changed to accept either a payload directly (in the form of FCompressedBuffer) or accept a reference to a IPayloadProvider which will return the payload when asked.
- Most of the existing backends make no use of the batched pushing feature and only push one payload at a time. The default implementation of this will be to ask each request one at a time for the payload and then pass it to the backend for a single push. If the IPayloadProvider is being used we will therefor only ever have one payload loaded at a time avoiding memory spikes.
- The source control backend currently does implement the batched pushing feature. This backend writes each payload to disk before submitting them to perforce, so all we have to do is make sure that request one payload, write it to disk, then discard the payload to avoid memory spikes.
- This change required that FPushRequest make it's members private and provide accessors to them, as the caller shouldn't need to care if the payload is loaded or comes via the provider.
- NOTE that this implementation is less efficient if we have packages containing many different payloads as we will end up opening the package, parsing the trailer then load a payload many times over, where as the original code would load all payloads from the file in one go. In practice, given the overhead of the source control backend this didn't make a huge difference and is probably not worth the effort to optimize at this point.

diff --git a/Engine/Source/Developer/Virtualization/Private/PackageSubmissionChecks.cpp b/Engine/Source/Developer/Virtualization/Private/PackageSubmissionChecks.cpp
--- a/Engine/Source/Developer/Virtualization/Private/PackageSubmissionChecks.cpp
+++ b/Engine/Source/Developer/Virtualization/Private/PackageSubmissionChecks.cpp
@@ -32,6 +32,132 @@
 namespace UE::Virtualization
 {
 
+/** 
+ * Implementation of the IPayloadProvider interface so that payloads can be requested on demand
+ * when they are being virtualized.
+ * 
+ * This implementation is not optimized. If a package holds many payloads that are all virtualized
+ * we will end up loading the same trailer over and over, as well as opening the same package file
+ * for read many times.
+ * 
+ * So far this has shown to be a rounding error compared to the actual cost of virtualization 
+ * and so implementing any level of caching has been left as a future task.
+ */
+class FWorkspaceDomainPayloadProvider final : public IPayloadProvider
+{
+public: 
+	FWorkspaceDomainPayloadProvider() = default;
+	virtual ~FWorkspaceDomainPayloadProvider() = default;
+
+	/** Register the payload with it's trailer and package name so that we can access it later as needed */
+	void RegisterPayload(const FIoHash& PayloadId, uint64 SizeOnDisk, const FString& PackageName)
+	{
+		if (!PayloadId.IsZero())
+		{
+			PayloadLookupTable.Emplace(PayloadId, FPayloadData(SizeOnDisk, PackageName));
+		}
+	}
+
+private:
+	virtual FCompressedBuffer RequestPayload(const FIoHash& Identifier) override
+	{
+		if (Identifier.IsZero())
+		{
+			return FCompressedBuffer();
+		}
+
+		const FPayloadData* Data = PayloadLookupTable.Find(Identifier);
+		if (Data == nullptr)
+		{
+			UE_LOG(LogVirtualization, Error, TEXT("FPayloadProvider was unable to find a payload with the identifier '%s'"), 
+				*LexToString(Identifier));
+
+			return FCompressedBuffer();
+		}
+		
+		TUniquePtr<FArchive> PackageAr = IPackageResourceManager::Get().OpenReadExternalResource(EPackageExternalResource::WorkspaceDomainFile, *Data->PackageName);
+
+		if (!PackageAr.IsValid())
+		{
+			UE_LOG(LogVirtualization, Error, TEXT("FPayloadProvider was unable to open the package '%s' for reading"), 
+				*Data->PackageName);
+
+			return FCompressedBuffer();
+		}
+			
+		PackageAr->Seek(PackageAr->TotalSize());
+
+		FPackageTrailer Trailer;
+		if (!Trailer.TryLoadBackwards(*PackageAr))
+		{
+			UE_LOG(LogVirtualization, Error, TEXT("FPayloadProvider failed to load the package trailer from the package '%s'"), 
+				*Data->PackageName);
+
+			return FCompressedBuffer();
+		}
+
+		FCompressedBuffer Payload = Trailer.LoadLocalPayload(Identifier, *PackageAr);
+		
+		if (!Payload)
+		{
+			UE_LOG(LogVirtualization, Error, TEXT("FPayloadProvider was uanble to load the payload '%s' from the package '%s'"),
+				*LexToString(Identifier),
+				*Data->PackageName);
+
+			return FCompressedBuffer();
+		}
+
+		if (Identifier != FIoHash(Payload.GetRawHash()))
+		{
+			UE_LOG(LogVirtualization, Error, TEXT("FPayloadProvider loaded an incorrect payload from the package '%s'. Expected '%s' Loaded  '%s'"), 
+				*Data->PackageName,
+				*LexToString(Identifier),
+				*LexToString(Payload.GetRawHash()));
+
+			return FCompressedBuffer();
+		}
+
+		return Payload;
+	}
+
+	virtual uint64 GetPayloadSize(const FIoHash& Identifier) override
+	{
+		if (Identifier.IsZero())
+		{
+			return 0;
+		}
+
+		const FPayloadData* Data = PayloadLookupTable.Find(Identifier);
+		if (Data != nullptr)
+		{
+			return Data->SizeOnDisk;
+		}
+		else
+		{
+			UE_LOG(LogVirtualization, Error, TEXT("FPayloadProvider was unable to find a payload with the identifier '%s'"),
+				*LexToString(Identifier));
+
+			return 0;
+		}
+	}
+
+	/* This structure holds additional info about the payload that we might need later */
+	struct FPayloadData
+	{
+		FPayloadData(uint64 InSizeOnDisk, const FString& InPackageName)
+			: SizeOnDisk(InSizeOnDisk)
+			, PackageName(InPackageName)
+		{
+
+		}
+
+		uint64 SizeOnDisk;
+		FString PackageName;
+	};
+
+	TMap<FIoHash, FPayloadData> PayloadLookupTable;
+};
+
 /**
  * Check that the given package ends with PACKAGE_FILE_TAG. Intended to be used to make sure that
  * we have truncated a package correctly when removing the trailers.
@@ -269,10 +395,12 @@ void VirtualizePackages(const TArray<FString>& FilesToSubmit, TArray<FText>& Out
 	// TODO Optimization: In theory we could have many packages sharing the same payload and we only need to push once
 	Progress.EnterProgressFrame(1.0f);
 
+	// Build up the info in the payload provider and the final array of payload push requests
+
+	FWorkspaceDomainPayloadProvider PayloadProvider;
 	TArray<Virtualization::FPushRequest> PayloadsToSubmit;
 	PayloadsToSubmit.Reserve(TotalPayloadsToVirtualize);
 
-	// Push any remaining local payload to the persistent backends
 	for (FPackageInfo& PackageInfo : Packages)
 	{
 		if (PackageInfo.LocalPayloads.IsEmpty())
@@ -280,50 +408,19 @@ void VirtualizePackages(const TArray<FString>& FilesToSubmit, TArray<FText>& Out
 			continue;
 		}
 
-		TUniquePtr<FArchive> PackageAr = IPackageResourceManager::Get().OpenReadExternalResource(EPackageExternalResource::WorkspaceDomainFile, PackageInfo.Path.GetPackageName());
-
-		if (!PackageAr.IsValid())
-		{
-			FText Message = FText::Format(	LOCTEXT("Virtualization_PkgOpen", "Failed to open the package '{1}' for reading"),
-											FText::FromString(PackageInfo.Path.GetDebugName()));
-			OutErrors.Add(Message);
-			return;
-		}
-
 		PackageInfo.PayloadIndex = PayloadsToSubmit.Num();
 
 		for (const FIoHash& PayloadId : PackageInfo.LocalPayloads)
 		{
-			checkf(!PayloadId.IsZero(), TEXT("PackageTrailer for package '%s' should not contain invalid FIoHashs"), *PackageInfo.Path.GetDebugName());
-			
-			FCompressedBuffer Payload = PackageInfo.Trailer.LoadLocalPayload(PayloadId, *PackageAr);
-
-			if (PayloadId != FIoHash(Payload.GetRawHash()))
-			{
-				FText Message = FText::Format(	LOCTEXT("Virtualization_WrongPayload", "Package {0} loaded an incorrect payload from the trailer. Expected '{1}' Loaded  '{2}'"),
-												FText::FromString(PackageInfo.Path.GetDebugName()),
-												FText::FromString(LexToString(PayloadId)),
-												FText::FromString(LexToString(Payload.GetRawHash())));
-				OutErrors.Add(Message);
-				return;
-			}
-
-			if (!Payload)
-			{
-				FText Message = FText::Format(	LOCTEXT("Virtualization_MissingPayload", "Unable to find the payload '{0}' in the local storage of package '{1}'"),
-												FText::FromString(LexToString(PayloadId)),
-												FText::FromString(PackageInfo.Path.GetDebugName()));
-				OutErrors.Add(Message);
-				return;
-
-			}
+			const uint64 SizeOnDisk = PackageInfo.Trailer.FindPayloadSizeOnDisk(PayloadId);
 
-			PayloadsToSubmit.Emplace(PayloadId, MoveTemp(Payload), PackageInfo.Path.GetDebugName());
+			PayloadProvider.RegisterPayload(PayloadId, SizeOnDisk, PackageInfo.Path.GetPackageName());
+			PayloadsToSubmit.Emplace(PayloadId, PayloadProvider, PackageInfo.Path.GetPackageName());
 		}
 	}
 
 	Progress.EnterProgressFrame(1.0f);
-
+	// Push any remaining local payload to the persistent backends
 	if (!System.PushData(PayloadsToSubmit, EStorageType::Persistent))
 	{
 		FText Message = LOCTEXT("Virtualization_PushFailure", "Failed to push payloads");
@@ -334,7 +431,7 @@ void VirtualizePackages(const TArray<FString>& FilesToSubmit, TArray<FText>& Out
 	int64 TotalPayloadsVirtualized = 0;
 	for (const Virtualization::FPushRequest& Request : PayloadsToSubmit)
 	{
-		TotalPayloadsVirtualized += Request.Status == FPushRequest::EStatus::Success ? 1 : 0;
+		TotalPayloadsVirtualized += Request.GetStatus() == FPushRequest::EStatus::Success ? 1 : 0;
 	}
 	UE_LOG(LogVirtualization, Display, TEXT("Pushed %" INT64_FMT " payload(s) to persistent virtualized storage"), TotalPayloadsVirtualized);
 
@@ -344,18 +441,18 @@ void VirtualizePackages(const TArray<FString>& FilesToSubmit, TArray<FText>& Out
 		for (int32 Index = 0; Index < PackageInfo.LocalPayloads.Num(); ++Index)
 		{
 			const Virtualization::FPushRequest& Request = PayloadsToSubmit[PackageInfo.PayloadIndex + Index];
-			check(Request.Identifier == PackageInfo.LocalPayloads[Index]);
+			check(Request.GetIdentifier() == PackageInfo.LocalPayloads[Index]);
 
-			if (Request.Status == Virtualization::FPushRequest::EStatus::Success)
+			if (Request.GetStatus() == Virtualization::FPushRequest::EStatus::Success)
 			{
-				if (PackageInfo.Trailer.UpdatePayloadAsVirtualized(Request.Identifier))
+				if (PackageInfo.Trailer.UpdatePayloadAsVirtualized(Request.GetIdentifier()))
 				{
 					PackageInfo.bWasTrailerUpdated = true;
 				}
 				else
 				{
 					FText Message = FText::Format(	LOCTEXT("Virtualization_UpdateStatusFailed", "Unable to update the status for the payload '{0}' in the package '{1}'"),
-													FText::FromString(LexToString(Request.Identifier)),
+													FText::FromString(LexToString(Request.GetIdentifier())),
 													FText::FromString(PackageInfo.Path.GetDebugName()));
 					OutErrors.Add(Message);
 					return;
diff --git a/Engine/Source/Developer/Virtualization/Private/VirtualizationManager.cpp b/Engine/Source/Developer/Virtualization/Private/VirtualizationManager.cpp
--- a/Engine/Source/Developer/Virtualization/Private/VirtualizationManager.cpp
+++ b/Engine/Source/Developer/Virtualization/Private/VirtualizationManager.cpp
@@ -363,31 +363,31 @@ bool FVirtualizationManager::PushData(TArrayView<FPushRequest> Requests, EStorag
 		OriginalToValidatedRequest[Index] = INDEX_NONE;
 
 		FPushRequest& Request = Requests[Index];
-		if (Request.Identifier.IsZero() || Request.Payload.GetCompressedSize() == 0)
+		if (Request.GetIdentifier().IsZero() || Request.GetPayloadSize() == 0)
 		{
-			Request.Status = FPushRequest::EStatus::Invalid;
+			Request.SetStatus(FPushRequest::EStatus::Invalid);
 			continue;
 		}
 
-		if ((int64)Request.Payload.GetCompressedSize() < MinPayloadLength)
+		if ((int64)Request.GetPayloadSize() < MinPayloadLength)
 		{
 			UE_LOG(	LogVirtualization, Verbose, TEXT("Pushing payload (id: %s) with context ('%s') was prevented as it is smaller (%" UINT64_FMT ") than the MinPayloadLength (%" INT64_FMT ")"),
-					*LexToString(Request.Identifier),
-					*Request.Context,
-					Request.Payload.GetCompressedSize(),
+					*LexToString(Request.GetIdentifier()),
+					*Request.GetContext(),
+					Request.GetPayloadSize(),
 					MinPayloadLength);
 
-			Request.Status = FPushRequest::EStatus::BelowMinSize;
+			Request.SetStatus(FPushRequest::EStatus::BelowMinSize);
 			continue;
 		}
 
-		if (!ShouldVirtualize(Request.Context))
+		if (!ShouldVirtualize(Request.GetContext()))
 		{
 			UE_LOG(	LogVirtualization, Verbose, TEXT("Pushing payload (id: %s) with context ('%s') was prevented by filtering"),
-					*LexToString(Request.Identifier), 
-					*Request.Context);
+					*LexToString(Request.GetIdentifier()), 
+					*Request.GetContext());
 			
-			Request.Status = FPushRequest::EStatus::ExcludedByPackagPath;
+			Request.SetStatus(FPushRequest::EStatus::ExcludedByPackagPath);
 			continue;
 		}
 
@@ -434,11 +434,11 @@ bool FVirtualizationManager::PushData(TArrayView<FPushRequest> Requests, EStorag
 		{
 			for (FPushRequest& Request : ValidatedRequests)
 			{
-				FCompressedBuffer ValidationPayload = PullDataFromBackend(*Backend, Request.Identifier);
-				checkf(	Request.Payload.GetRawHash() == ValidationPayload.GetRawHash(),
+				FCompressedBuffer ValidationPayload = PullDataFromBackend(*Backend, Request.GetIdentifier());
+				checkf(	Request.GetIdentifier() == ValidationPayload.GetRawHash(),
 						TEXT("[%s] Failed to pull payload '%s' after it was pushed to backend"),
 						*Backend->GetDebugName(),
-						*LexToString(Request.Identifier));
+						*LexToString(Request.GetIdentifier()));
 			}
 		}
 	}
@@ -451,7 +451,7 @@ bool FVirtualizationManager::PushData(TArrayView<FPushRequest> Requests, EStorag
 		const int32 MappingIndex = OriginalToValidatedRequest[Index];
 		if (MappingIndex != INDEX_NONE)
 		{
-			Requests[Index].Status = ValidatedRequests[MappingIndex].Status;
+			Requests[Index].SetStatus(ValidatedRequests[MappingIndex].GetStatus());
 		}
 	}
 
@@ -1201,10 +1201,10 @@ bool FVirtualizationManager::TryPushDataToBackend(IVirtualizationBackend& Backen
 		for (const FPushRequest& Request : Requests)
 		{
 			// TODO: Don't add a hit if the payload was already uploaded
-			if (Request.Status == FPushRequest::EStatus::Success)
+			if (Request.GetStatus() == FPushRequest::EStatus::Success)
 			{
 				Stats.Accumulate(FCookStats::CallStats::EHitOrMiss::Hit, FCookStats::CallStats::EStatType::Counter, 1l, bIsInGameThread);	
-				Stats.Accumulate(FCookStats::CallStats::EHitOrMiss::Hit, FCookStats::CallStats::EStatType::Bytes, Request.Payload.GetCompressedSize(), bIsInGameThread);
+				Stats.Accumulate(FCookStats::CallStats::EHitOrMiss::Hit, FCookStats::CallStats::EStatType::Bytes, Request.GetPayloadSize(), bIsInGameThread);
 			}
 		}
 	}
diff --git a/Engine/Source/Developer/Virtualization/Private/VirtualizationSourceControlBackend.cpp b/Engine/Source/Developer/Virtualization/Private/VirtualizationSourceControlBackend.cpp
--- a/Engine/Source/Developer/Virtualization/Private/VirtualizationSourceControlBackend.cpp
+++ b/Engine/Source/Developer/Virtualization/Private/VirtualizationSourceControlBackend.cpp
@@ -176,7 +176,7 @@ static void CreateDescription(const FString& ProjectName, const TArray<const FPu
 
 	for (const FPushRequest* Request : FileRequests)
 	{
-		if (!Request->Context.IsEmpty())
+		if (!Request->GetContext().IsEmpty())
 		{
 			if (!bInitialNewline)
 			{
@@ -184,7 +184,7 @@ static void CreateDescription(const FString& ProjectName, const TArray<const FPu
 				bInitialNewline = true;
 			}
 
-			OutDescription << TEXT("\n") << Request->Identifier << "\t: " << Request->Context;
+			OutDescription << TEXT("\n") << Request->GetIdentifier() << "\t: " << Request->GetContext();
 		}
 	}
 }
@@ -437,13 +437,24 @@ bool FSourceControlBackend::PushData(TArrayView<FPushRequest> Requests)
 		for (const FPushRequest& Request : Requests)
 		{
 			TStringBuilder<52> LocalPayloadPath;
-			Utils::PayloadIdToPath(Request.Identifier, LocalPayloadPath);
+			Utils::PayloadIdToPath(Request.GetIdentifier(), LocalPayloadPath);
 
 			TStringBuilder<260> PayloadFilePath;
 			FPathViews::Append(PayloadFilePath, SessionDirectory, LocalPayloadPath);
 
 			UE_LOG(LogVirtualization, Verbose, TEXT("[%s] Writing payload to '%s' for submission"), *GetDebugName(), PayloadFilePath.ToString());
 
+			FCompressedBuffer Payload = Request.GetPayload();
+			if (Payload.IsNull())
+			{
+				UE_LOG(LogVirtualization, Error, TEXT("[%s] Failed to acquire payload '%s' contents to '%s' for writing"),
+					*GetDebugName(),
+					*LexToString(Request.GetIdentifier()),
+					PayloadFilePath.ToString());
+
+				return false;
+			}
+
 			TUniquePtr<FArchive> FileAr(IFileManager::Get().CreateFileWriter(PayloadFilePath.ToString()));
 			if (!FileAr)
 			{
@@ -452,14 +463,14 @@ bool FSourceControlBackend::PushData(TArrayView<FPushRequest> Requests)
 
 				UE_LOG(LogVirtualization, Error, TEXT("[%s] Failed to write payload '%s' contents to '%s' due to system error: %s"),
 					*GetDebugName(),
-					*LexToString(Request.Identifier),
+					*LexToString(Request.GetIdentifier()),
 					PayloadFilePath.ToString(),
 					SystemErrorMsg.ToString());
 
 				return false;
 			}
 
-			Request.Payload.Save(*FileAr);
+			Payload.Save(*FileAr);
 
 			if (!FileAr->Close())
 			{
@@ -468,7 +479,7 @@ bool FSourceControlBackend::PushData(TArrayView<FPushRequest> Requests)
 
 				UE_LOG(LogVirtualization, Error, TEXT("[%s] Failed to write payload '%s' contents to '%s' due to system error: %s"),
 					*GetDebugName(),
-					*LexToString(Request.Identifier),
+					*LexToString(Request.GetIdentifier()),
 					*PayloadFilePath,
 					SystemErrorMsg.
 					ToString());
@@ -578,7 +589,7 @@ bool FSourceControlBackend::PushData(TArrayView<FPushRequest> Requests)
 		{
 			// TODO: Maybe check if the data is the same (could be different if the compression algorithm has changed)
 			// TODO: Should we respect if the file is deleted as technically we can still get access to it?
-			Requests[Index].Status = FPushRequest::EStatus::Success;
+			Requests[Index].SetStatus(FPushRequest::EStatus::Success);
 		}
 		else if (FileStates[Index]->CanAdd())
 		{
@@ -630,7 +641,7 @@ bool FSourceControlBackend::PushData(TArrayView<FPushRequest> Requests)
 	// TODO: We really should be setting a more fine grain status for each request, or not bother with the status at all
 	for (FPushRequest& Request : Requests)
 	{
-		Request.Status = FPushRequest::EStatus::Success;
+		Request.SetStatus(FPushRequest::EStatus::Success);
 	}
 
 	return true;
diff --git a/Engine/Source/Developer/Virtualization/Public/IVirtualizationBackend.h b/Engine/Source/Developer/Virtualization/Public/IVirtualizationBackend.h
--- a/Engine/Source/Developer/Virtualization/Public/IVirtualizationBackend.h
+++ b/Engine/Source/Developer/Virtualization/Public/IVirtualizationBackend.h
@@ -95,21 +95,21 @@ public:
 		// TODO: Improve the error codes in the future
 		for (FPushRequest& Request : Requests)
 		{
-			EPushResult Result = PushData(Request.Identifier, Request.Payload, Request.Context);
+			EPushResult Result = PushData(Request.GetIdentifier(), Request.GetPayload(), Request.GetContext());
 			switch (Result)
 			{
 			case EPushResult::Failed:
-				Request.Status = FPushRequest::EStatus::Failed;
+				Request.SetStatus(FPushRequest::EStatus::Failed);
 				return false;
 
 			case EPushResult::PayloadAlreadyExisted:
 				// falls through
 			case EPushResult::Success:
-				Request.Status = FPushRequest::EStatus::Success;
+				Request.SetStatus(FPushRequest::EStatus::Success);
 				break;
 
 			default:
-				Request.Status = FPushRequest::EStatus::Failed;
+				Request.SetStatus(FPushRequest::EStatus::Failed);
 				checkNoEntry();
 				break;
 			}
diff --git a/Engine/Source/Runtime/Core/Public/Virtualization/VirtualizationSystem.h b/Engine/Source/Runtime/Core/Public/Virtualization/VirtualizationSystem.h
--- a/Engine/Source/Runtime/Core/Public/Virtualization/VirtualizationSystem.h
+++ b/Engine/Source/Runtime/Core/Public/Virtualization/VirtualizationSystem.h
@@ -68,7 +68,41 @@ enum class FPayloadStatus : int8
 	FoundAll
 };
 
-/** Data structure representing a request to push a payload to a backend storage system */
+/** 
+ * This interface can be implemented and passed to a FPushRequest as a way of providing the payload 
+ * to the virtualization system for a push operation but deferring the loading of the payload from 
+ * disk until it is actually needed. In some cases this allows the loading of the payload to be 
+ * skipped entirely (if the payload is already in all backends for example) or can prevent memory
+ * spikes caused by loading a large number of payloads for a batched push request.
+ * 
+ * Note that if the backend graph contains multiple backends then payloads may be requested 
+ * multiple times. It will be up to the provider implementation to decide if a requested
+ * payload should be cached in case of future access or not. The methods are not const in order
+ * to make it easier for derived classes to cache the results if needed without the use of
+ * mutable.
+ */
+class IPayloadProvider
+{
+public:
+	IPayloadProvider() = default;
+	virtual ~IPayloadProvider() = default;
+
+	/** 
+	 * Should return the payload for the given FIoHash. If the provider fails to find the payload
+	 * then it should return a null FCompressedBuffer to indicate the error.
+	 */
+	virtual FCompressedBuffer RequestPayload(const FIoHash& Identifier) = 0;
+
+	/** Returns the current size of the payload on disk. */
+	virtual uint64 GetPayloadSize(const FIoHash& Identifier) = 0;
+};
+
+/** 
+ * Data structure representing a request to push a payload to a backend storage system. 
+ * Note that a request can either before for payload already in memory (in which case the
+ * payload should be passed into the constructor as a FCompressedBuffer) or by a 
+ * IPayloadProvider which will provide the payload on demand.
+*/
 struct FPushRequest
 {
 	enum class EStatus
@@ -85,6 +119,16 @@ struct FPushRequest
 		Success
 	};
 
+	FPushRequest() = delete;
+	~FPushRequest() = default;
+
+	/** 
+	 * Create a request for a payload already in memory 
+	 *
+	 * @param InIdentifier	The hash of the payload in its uncompressed form
+	 * @param InPayload		The payload, this can be in any compressed format that the caller wishes.
+	 * @param InContext		Content showing where the payload came from. If it comes from a package then this should be the package path
+	 */
 	FPushRequest(const FIoHash& InIdentifier, const FCompressedBuffer& InPayload, const FString& InContext)
 		: Identifier(InIdentifier)
 		, Payload(InPayload)
@@ -93,6 +137,13 @@ struct FPushRequest
 
 	}
 
+	/**
+	 * Create a request for a payload already in memory
+	 *
+	 * @param InIdentifier	The hash of the payload in its uncompressed form
+	 * @param InPayload		The payload, this can be in any compressed format that the caller wishes.
+	 * @param InContext		Content showing where the payload came from. If it comes from a package then this should be the package path
+	 */
 	FPushRequest(const FIoHash& InIdentifier, FCompressedBuffer&& InPayload, FString&& InContext)
 		: Identifier(InIdentifier)
 		, Payload(InPayload)
@@ -101,10 +152,81 @@ struct FPushRequest
 
 	}
 
+	/**
+	 * Create a request for a payload to be loaded on demand
+	 *
+	 * @param InIdentifier	The hash of the payload in its uncompressed form
+	 * @param InProvider	The provider that will load the payload when requested. The providers lifespan must exceed that of the FPushRequest
+	 * @param InContext		Content showing where the payload came from. If it comes from a package then this should be the package path
+	 */
+	FPushRequest(const FIoHash& InIdentifier, IPayloadProvider& InProvider, FString&& InContext)
+		: Identifier(InIdentifier)
+		, Provider(&InProvider)
+		, Context(InContext)
+	{
+
+	}
+
+	/** Return the identifer used in the request */
+	FIoHash GetIdentifier() const
+	{
+		return Identifier;
+	}
+
+	/** Returns the current status of the request */
+	EStatus GetStatus() const
+	{
+		return Status;
+	}
+
+	/** Returns the size of the payload when it was on disk */
+	uint64 GetPayloadSize() const
+	{
+		if (Provider != nullptr)
+		{
+			return Provider->GetPayloadSize(Identifier);
+		}
+		else
+		{
+			return Payload.GetCompressedSize();
+		}
+	}
+	
+	/** Returns the payload */
+	FCompressedBuffer GetPayload() const
+	{
+		if (Provider != nullptr)
+		{
+			return Provider->RequestPayload(Identifier);
+		}
+		else
+		{
+			return Payload;
+		}
+	}
+
+	/** Returns the context of the payload */
+	const FString& GetContext() const
+	{
+		return Context;
+	}
+
+	/** Allows the status of the request to be set, this should only be done by the virtualization backends */
+	void SetStatus(EStatus InStatus)
+	{
+		Status = InStatus;
+	}
+
+private:
 	/** The identifier of the payload */
 	FIoHash Identifier;
+
 	/** The payload data */
 	FCompressedBuffer Payload;
+
+	/** Provider to retrieve the payload from */
+	IPayloadProvider* Provider = nullptr;
+
 	/** A string containing context for the payload, typically a package name */
 	FString Context;
 
