Allow the source control backend to split large request batches into smaller ones

#rnx

- When submitting large batches it was found that once a cl description reached a certain length, that p4v would lock up attempting to display it in the history tab, so we want to be able to make smaller submits to keep the description to a managable size.
- In addition it is better for our server set up to make multiple smaller batch submits than insanely large ones.
- The max size of the batch can be set when creating the virtualization graph in the config file.

diff --git a/Engine/Source/Developer/Virtualization/Private/VirtualizationSourceControlBackend.cpp b/Engine/Source/Developer/Virtualization/Private/VirtualizationSourceControlBackend.cpp
--- a/Engine/Source/Developer/Virtualization/Private/VirtualizationSourceControlBackend.cpp
+++ b/Engine/Source/Developer/Virtualization/Private/VirtualizationSourceControlBackend.cpp
@@ -407,7 +407,7 @@ bool FSourceControlBackend::PushData(TArrayView<FPushRequest> Requests)
 
 	const FGuid SessionGuid = FGuid::NewGuid();
 	
-	UE_LOG(LogVirtualization, Display, TEXT("[%s] Started payload submission session '%s' for '%d' payload(s)"), *GetDebugName(), *LexToString(SessionGuid), Requests.Num());
+	UE_LOG(LogVirtualization, Log, TEXT("[%s] Started payload submission session '%s' for '%d' payload(s)"), *GetDebugName(), *LexToString(SessionGuid), Requests.Num());
 
 	TStringBuilder<260> SessionDirectory;
 	FPathViews::Append(SessionDirectory, SubmissionRootDir, SessionGuid);
@@ -418,7 +418,7 @@ bool FSourceControlBackend::PushData(TArrayView<FPushRequest> Requests)
 		return false;
 	}
 
-	UE_LOG(LogVirtualization, Display, TEXT("[%s] Created directory '%s' to submit payloads from"), *GetDebugName(), SessionDirectory.ToString());
+	UE_LOG(LogVirtualization, Log, TEXT("[%s] Created directory '%s' to submit payloads from"), *GetDebugName(), SessionDirectory.ToString());
 
 	ON_SCOPE_EXIT
 	{
@@ -426,73 +426,6 @@ bool FSourceControlBackend::PushData(TArrayView<FPushRequest> Requests)
 		IFileManager::Get().DeleteDirectory(SessionDirectory.ToString(), false, true);
 	};
 
-	TArray<FString> FilesToSubmit;
-	FilesToSubmit.Reserve(Requests.Num());
-
-	// Write the payloads to disk so that they can be submitted (source control module currently requires the files to
-	// be on disk)
-	{
-		TRACE_CPUPROFILER_EVENT_SCOPE(FSourceControlBackend::PushData::CreateFiles);
-
-		for (const FPushRequest& Request : Requests)
-		{
-			TStringBuilder<52> LocalPayloadPath;
-			Utils::PayloadIdToPath(Request.GetIdentifier(), LocalPayloadPath);
-
-			TStringBuilder<260> PayloadFilePath;
-			FPathViews::Append(PayloadFilePath, SessionDirectory, LocalPayloadPath);
-
-			UE_LOG(LogVirtualization, Verbose, TEXT("[%s] Writing payload to '%s' for submission"), *GetDebugName(), PayloadFilePath.ToString());
-
-			FCompressedBuffer Payload = Request.GetPayload();
-			if (Payload.IsNull())
-			{
-				UE_LOG(LogVirtualization, Error, TEXT("[%s] Failed to acquire payload '%s' contents to '%s' for writing"),
-					*GetDebugName(),
-					*LexToString(Request.GetIdentifier()),
-					PayloadFilePath.ToString());
-
-				return false;
-			}
-
-			TUniquePtr<FArchive> FileAr(IFileManager::Get().CreateFileWriter(PayloadFilePath.ToString()));
-			if (!FileAr)
-			{
-				TStringBuilder<MAX_SPRINTF> SystemErrorMsg;
-				Utils::GetFormattedSystemError(SystemErrorMsg);
-
-				UE_LOG(LogVirtualization, Error, TEXT("[%s] Failed to write payload '%s' contents to '%s' due to system error: %s"),
-					*GetDebugName(),
-					*LexToString(Request.GetIdentifier()),
-					PayloadFilePath.ToString(),
-					SystemErrorMsg.ToString());
-
-				return false;
-			}
-
-			Payload.Save(*FileAr);
-
-			if (!FileAr->Close())
-			{
-				TStringBuilder<MAX_SPRINTF> SystemErrorMsg;
-				Utils::GetFormattedSystemError(SystemErrorMsg);
-
-				UE_LOG(LogVirtualization, Error, TEXT("[%s] Failed to write payload '%s' contents to '%s' due to system error: %s"),
-					*GetDebugName(),
-					*LexToString(Request.GetIdentifier()),
-					*PayloadFilePath,
-					SystemErrorMsg.
-					ToString());
-
-				return false;
-			}
-
-			FilesToSubmit.Emplace(MoveTemp(PayloadFilePath));
-		}
-	}
-
-	check(Requests.Num() == FilesToSubmit.Num());
-
 	FSemaphoreScopeLock _(ConcurrentConnectionLimit.Get());
 
 	TStringBuilder<64> WorkspaceName;
@@ -566,82 +499,173 @@ bool FSourceControlBackend::PushData(TArrayView<FPushRequest> Requests)
 		}
 	};
 
-	TArray<FSourceControlStateRef> FileStates;
+	const int32 NumBatches = FMath::DivideAndRoundUp<int32>(Requests.Num(), MaxBatchCount);
+
+	UE_LOG(LogVirtualization, Log, TEXT("[%s] Splitting the push into '%d' batches"), *GetDebugName(), NumBatches);
+
+	for (int32 BatchIndex = 0; BatchIndex < NumBatches; ++BatchIndex)
 	{
-		TRACE_CPUPROFILER_EVENT_SCOPE(FSourceControlBackend::PushData::GetFileStates);
-		if (GetDepotPathStates(*SCCProvider, FilesToSubmit, FileStates) != ECommandResult::Succeeded)
-		{
-			UE_LOG(LogVirtualization, Error, TEXT("[%s] Failed to find the current file state for payloads"), *GetDebugName());
-			return false;
-		}
-	}
-	check(Requests.Num() == FileStates.Num());
+		UE_LOG(LogVirtualization, Log, TEXT("[%s] Processing batch %d/%d..."), *GetDebugName(), BatchIndex + 1, NumBatches);
 
-	TArray<FString> FilesToAdd;
-	FilesToAdd.Reserve(FilesToSubmit.Num());
+		const int32 BatchStart = BatchIndex * MaxBatchCount;
+		const int32 BatchEnd = FMath::Min((BatchIndex + 1) * MaxBatchCount, Requests.Num());
 
-	TArray<const FPushRequest*> FileRequests;
-	FileRequests.Reserve(FilesToSubmit.Num());
+		TArrayView<FPushRequest> RequestBatch(&Requests[BatchStart], BatchEnd - BatchStart);
 
-	for (int32 Index = 0; Index < FilesToSubmit.Num(); ++Index)
-	{
-		if (FileStates[Index]->IsSourceControlled())
+		TArray<FString> FilesToSubmit;
+		FilesToSubmit.Reserve(RequestBatch.Num());
+
+		// Write the payloads to disk so that they can be submitted (source control module currently requires the files to
+		// be on disk)
 		{
-			// TODO: Maybe check if the data is the same (could be different if the compression algorithm has changed)
-			// TODO: Should we respect if the file is deleted as technically we can still get access to it?
-			Requests[Index].SetStatus(FPushRequest::EStatus::Success);
+			TRACE_CPUPROFILER_EVENT_SCOPE(FSourceControlBackend::PushData::CreateFiles);
+
+			for (const FPushRequest& Request : RequestBatch)
+			{
+				TStringBuilder<52> LocalPayloadPath;
+				Utils::PayloadIdToPath(Request.GetIdentifier(), LocalPayloadPath);
+
+				TStringBuilder<260> PayloadFilePath;
+				FPathViews::Append(PayloadFilePath, SessionDirectory, LocalPayloadPath);
+
+				UE_LOG(LogVirtualization, Verbose, TEXT("[%s] Writing payload to '%s' for submission"), *GetDebugName(), PayloadFilePath.ToString());
+
+				FCompressedBuffer Payload = Request.GetPayload();
+				if (Payload.IsNull())
+				{
+					UE_LOG(LogVirtualization, Error, TEXT("[%s] Failed to acquire payload '%s' contents to '%s' for writing"),
+						*GetDebugName(),
+						*LexToString(Request.GetIdentifier()),
+						PayloadFilePath.ToString());
+
+					return false;
+				}
+
+				TUniquePtr<FArchive> FileAr(IFileManager::Get().CreateFileWriter(PayloadFilePath.ToString()));
+				if (!FileAr)
+				{
+					TStringBuilder<MAX_SPRINTF> SystemErrorMsg;
+					Utils::GetFormattedSystemError(SystemErrorMsg);
+
+					UE_LOG(LogVirtualization, Error, TEXT("[%s] Failed to write payload '%s' contents to '%s' due to system error: %s"),
+						*GetDebugName(),
+						*LexToString(Request.GetIdentifier()),
+						PayloadFilePath.ToString(),
+						SystemErrorMsg.ToString());
+
+					return false;
+				}
+
+				Payload.Save(*FileAr);
+
+				if (!FileAr->Close())
+				{
+					TStringBuilder<MAX_SPRINTF> SystemErrorMsg;
+					Utils::GetFormattedSystemError(SystemErrorMsg);
+
+					UE_LOG(LogVirtualization, Error, TEXT("[%s] Failed to write payload '%s' contents to '%s' due to system error: %s"),
+						*GetDebugName(),
+						*LexToString(Request.GetIdentifier()),
+						*PayloadFilePath,
+						SystemErrorMsg.
+						ToString());
+
+					return false;
+				}
+
+				FilesToSubmit.Emplace(MoveTemp(PayloadFilePath));
+			}
 		}
-		else if (FileStates[Index]->CanAdd())
+
+		check(RequestBatch.Num() == FilesToSubmit.Num());
+
+		TArray<FSourceControlStateRef> FileStates;
 		{
-			FilesToAdd.Add(FilesToSubmit[Index]);
-			FileRequests.Add(&Requests[Index]);
+			TRACE_CPUPROFILER_EVENT_SCOPE(FSourceControlBackend::PushData::GetFileStates);
+			if (GetDepotPathStates(*SCCProvider, FilesToSubmit, FileStates) != ECommandResult::Succeeded)
+			{
+				UE_LOG(LogVirtualization, Error, TEXT("[%s] Failed to find the current file state for payloads"), *GetDebugName());
+				return false;
+			}
 		}
-		else
+		check(RequestBatch.Num() == FileStates.Num());
+
+		TArray<FString> FilesToAdd;
+		FilesToAdd.Reserve(FilesToSubmit.Num());
+
+		TArray<const FPushRequest*> FileRequests;
+		FileRequests.Reserve(FilesToSubmit.Num());
+
+		for (int32 FileIndex = 0; FileIndex < FilesToSubmit.Num(); ++FileIndex)
 		{
-			UE_LOG(LogVirtualization, Error, TEXT("[%s] The the payload file '%s' is not in source control but also cannot be marked for Add"), *GetDebugName(), *FilesToSubmit[Index]);
-			return false;
+			if (FileStates[FileIndex]->IsSourceControlled())
+			{
+				// TODO: Maybe check if the data is the same (could be different if the compression algorithm has changed)
+				// TODO: Should we respect if the file is deleted as technically we can still get access to it?
+				RequestBatch[FileIndex].SetStatus(FPushRequest::EStatus::Success);
+			}
+			else if (FileStates[FileIndex]->CanAdd())
+			{
+				FilesToAdd.Add(FilesToSubmit[FileIndex]);
+				FileRequests.Add(&RequestBatch[FileIndex]);
+			}
+			else
+			{
+				UE_LOG(LogVirtualization, Error, TEXT("[%s] The the payload file '%s' is not in source control but also cannot be marked for Add"), *GetDebugName(), *FilesToSubmit[FileIndex]);
+				return false;
+			}
 		}
-	}
-
-	check(FileRequests.Num() == FilesToAdd.Num());
 
-	if (FilesToAdd.IsEmpty())
-	{
-		return true;
-	}
+		check(FileRequests.Num() == FilesToAdd.Num());
 
-	{
-		TRACE_CPUPROFILER_EVENT_SCOPE(FSourceControlBackend::PushData::AddFiles);
+		if (FilesToAdd.IsEmpty())
+		{
+			return true;
+		}
 
-		if (SCCProvider->Execute(ISourceControlOperation::Create<FMarkForAdd>(), FilesToAdd) != ECommandResult::Succeeded)
 		{
-			UE_LOG(LogVirtualization, Error, TEXT("[%s] Failed to mark the payload file for Add in source control"), *GetDebugName());
-			return false;
+			TRACE_CPUPROFILER_EVENT_SCOPE(FSourceControlBackend::PushData::AddFiles);
+
+			if (SCCProvider->Execute(ISourceControlOperation::Create<FMarkForAdd>(), FilesToAdd) != ECommandResult::Succeeded)
+			{
+				UE_LOG(LogVirtualization, Error, TEXT("[%s] Failed to mark the payload file for Add in source control"), *GetDebugName());
+				return false;
+			}
 		}
-	}
 
-	// Now submit the payload
-	{
-		TRACE_CPUPROFILER_EVENT_SCOPE(FSourceControlBackend::PushData::SubmitFiles);
+		// Now submit the payload
+		{
+			TRACE_CPUPROFILER_EVENT_SCOPE(FSourceControlBackend::PushData::SubmitFiles);
+
+			TSharedRef<FCheckIn, ESPMode::ThreadSafe> CheckInOperation = ISourceControlOperation::Create<FCheckIn>();
 
-		TSharedRef<FCheckIn, ESPMode::ThreadSafe> CheckInOperation = ISourceControlOperation::Create<FCheckIn>();
+			TStringBuilder<512> Description;
+			CreateDescription(ProjectName, FileRequests, Description);
 
-		TStringBuilder<512> Description;
-		CreateDescription(ProjectName, FileRequests, Description);
+			CheckInOperation->SetDescription(FText::FromString(Description.ToString()));
 
-		CheckInOperation->SetDescription(FText::FromString(Description.ToString()));
+			if (SCCProvider->Execute(CheckInOperation, FilesToAdd) != ECommandResult::Succeeded)
+			{
+				UE_LOG(LogVirtualization, Error, TEXT("[%s] Failed to submit the payload file(s) to source control"), *GetDebugName());
+				return false;
+			}
+		}
 
-		if (SCCProvider->Execute(CheckInOperation, FilesToAdd) != ECommandResult::Succeeded)
+		// TODO: We really should be setting a more fine grain status for each request, or not bother with the status at all
+		for (FPushRequest& Request : RequestBatch)
 		{
-			UE_LOG(LogVirtualization, Error, TEXT("[%s] Failed to submit the payload file(s) to source control"), *GetDebugName());
-			return false;
+			Request.SetStatus(FPushRequest::EStatus::Success);
 		}
-	}
 
-	// TODO: We really should be setting a more fine grain status for each request, or not bother with the status at all
-	for (FPushRequest& Request : Requests)
-	{
-		Request.SetStatus(FPushRequest::EStatus::Success);
+		// Try to clean up the files from this batch
+		for (const FString& FilePath : FilesToSubmit)
+		{
+			const bool bRequireExists = false;
+			const bool bEvenReadOnly = true;
+			const bool bQuiet = false;
+
+			IFileManager::Get().Delete(*FilePath, bRequireExists, bEvenReadOnly, bQuiet);
+		}
 	}
 
 	return true;
@@ -736,7 +760,18 @@ bool FSourceControlBackend::TryApplySettingsFromConfigFiles(const FString& Confi
 		{
 			UE_LOG(LogVirtualization, Log, TEXT("[%s] Has no limit to it's concurrent source control connections"), *GetDebugName());
 		}
-	}	
+	}
+
+	// Check for the optional BatchCount parameter
+	{
+		int32 MaxBatchCountIniFile = 0;
+		if (FParse::Value(*ConfigEntry, TEXT("MaxBatchCount="), MaxBatchCountIniFile))
+		{
+			MaxBatchCount = MaxBatchCountIniFile;
+		}
+
+		UE_LOG(LogVirtualization, Log, TEXT("[%s] Will push payloads in batches of up to %d payloads(s) at a time"), *GetDebugName(), MaxBatchCount);
+	}
 
 	if (!FindSubmissionWorkingDir(ConfigEntry))
 	{
diff --git a/Engine/Source/Developer/Virtualization/Private/VirtualizationSourceControlBackend.h b/Engine/Source/Developer/Virtualization/Private/VirtualizationSourceControlBackend.h
--- a/Engine/Source/Developer/Virtualization/Private/VirtualizationSourceControlBackend.h
+++ b/Engine/Source/Developer/Virtualization/Private/VirtualizationSourceControlBackend.h
@@ -40,6 +40,9 @@ class FSemaphore;
  *								in milliseconds. Remember that the max length of time that the process
  *								can stall attempting to download a payload file is 
  *								RetryCount * RetryWaitTime; [Default=100ms]
+ * BatchCount [int32]			The max number of payloads that can be pushed to source control in a
+ *								single submit. If the number of payloads in a request batch exceeds
+ *								this size then it will be split into multiple smaller batches. [Default=100]
  * 
  * Environment Variables:
  * UE-VirtualizationWorkingDir [string]:	This can be set to a valid directory path that the backend
@@ -93,6 +96,8 @@ private:
 	/** Should we try to make the temp client partitioned or not? */
 	bool bUsePartitionedClient = true;
 
+	int32 MaxBatchCount = 100;
+
 	/** A counted semaphore that will limit the number of concurrent connections that we can make */
 	TUniquePtr<UE::Virtualization::FSemaphore> ConcurrentConnectionLimit;
 	
