Add a VeryVerbose log message when pulling a payload to state the payload identifier and the backend it was pulled from.

#rnx

- Adding this so that we can more easily track in a build machine log when a payload gets pulled when debugging unexpected behavior.

diff --git a/Engine/Source/Developer/Virtualization/Private/VirtualizationManager.cpp b/Engine/Source/Developer/Virtualization/Private/VirtualizationManager.cpp
--- a/Engine/Source/Developer/Virtualization/Private/VirtualizationManager.cpp
+++ b/Engine/Source/Developer/Virtualization/Private/VirtualizationManager.cpp
@@ -1239,6 +1239,8 @@ FCompressedBuffer FVirtualizationManager::PullDataFromAllBackends(const FIoHash&
 				CachePayload(Id, Payload, Backend);
 			}
 
+			UE_LOG(LogVirtualization, VeryVerbose, TEXT("[%s] pulled payload '%s'"), *Backend->GetDebugName(), *LexToString(Id));
+
 			return Payload;
 		}
 	}
