We now warn users when trying to virtualize packages that are locked and cannot be modified rather than error

#rnx

### Problem
- When using the stand alone virtualization tool, to virtualize packages directly from p4v it is possible that the editor will have locked the package files. This means we cannot modify it and remove the payload making the virtualization of the package file impossible.
- At some point in the future we will add the ability for the tool to communicate with the editor to ask that the lock be relinquished OR there is the possibility that we will be moving away from the editor permantly locking package files.
- As this problem will be eventually be solved we'd rather this problem not present to the end users as an error as this will cause them to view the tool as annoying and fiddly to use.
- So for now we will treat this problem as a warning and allow the end user to continue and submit the package non-virtualized.

### Fix
- At the moment once we detect a package has local payloads that might need virtualization we check to see if the file could be written to.
- This might raise some false positives if the local payloads are filtered out or otherwise not suppose to be virtualized.
- At the moment this sort of filtering is done when pushing the payloads meaning we could end up pushing payloads then being unable to modify the package.
- A future work item will be added to clean this up at some point in the future (depending on when we will start work to remove the problem entirely)

diff --git a/Engine/Source/Developer/Virtualization/Private/PackageSubmissionChecks.cpp b/Engine/Source/Developer/Virtualization/Private/PackageSubmissionChecks.cpp
--- a/Engine/Source/Developer/Virtualization/Private/PackageSubmissionChecks.cpp
+++ b/Engine/Source/Developer/Virtualization/Private/PackageSubmissionChecks.cpp
@@ -250,6 +250,14 @@ bool TryCopyPackageWithoutTrailer(const FPackagePath PackagePath, const FString&
 	return true;
 }
 
+/** Tests if we would be able to write to the given file if we wanted to */
+bool CanWriteToFile(const FString& FilePath)
+{
+	TUniquePtr<FArchive> FileHandle(IFileManager::Get().CreateFileWriter(*FilePath, FILEWRITE_Append | FILEWRITE_Silent));
+
+	return FileHandle.IsValid();
+}
+
 void VirtualizePackages(const TArray<FString>& FilesToSubmit, TArray<FText>& OutDescriptionTags, TArray<FText>& OutErrors)
 {
 	TRACE_CPUPROFILER_EVENT_SCOPE(UE::Virtualization::VirtualizePackages);
@@ -327,10 +335,24 @@ void VirtualizePackages(const TArray<FString>& FilesToSubmit, TArray<FText>& Out
 				PkgInfo.Trailer = MoveTemp(Trailer);
 				PkgInfo.LocalPayloads = PkgInfo.Trailer.GetPayloads(EPayloadFilter::Local);
 
-				TotalPayloadsToCheck += PkgInfo.LocalPayloads.Num();
-
 				if (!PkgInfo.LocalPayloads.IsEmpty())
-				{
+				{	
+					if (!CanWriteToFile(AbsoluteFilePath))
+					{
+						// Technically the package could have local payloads that won't be virtualized due to filtering or min payload sizes and so the
+						// following warning is misleading. This will be solved if we move that evaluation to the point of saving a package.
+						// If not then we probably need to extend QueryPayloadStatuses to test filtering etc as well, then check for potential package
+						// modification after that.
+						// Long term, the stand alone tool should be able to request the UnrealEditor relinquish the lock on the package file so this becomes 
+						// less of a problem.
+						FText Message = FText::Format(LOCTEXT("Virtualization_PkgLocked", "The package file '{0}' has local payloads but is locked for modification and cannot be virtualized, this package will be skipped!"),
+							FText::FromString(PkgInfo.Path.GetDebugName()));
+						UE_LOG(LogVirtualization, Warning, TEXT("%s"), *Message.ToString());
+						continue;
+					}
+
+					TotalPayloadsToCheck += PkgInfo.LocalPayloads.Num();
+
 					PkgInfo.PayloadIndex = AllLocalPayloads.Num();
 					AllLocalPayloads.Append(PkgInfo.LocalPayloads);
 
