Move the critical section from a member of FVirtualizationManager to FDebugValues as it for debugging functionality only.

#rnx

diff --git a/Engine/Source/Developer/Virtualization/Private/VirtualizationManager.cpp b/Engine/Source/Developer/Virtualization/Private/VirtualizationManager.cpp
--- a/Engine/Source/Developer/Virtualization/Private/VirtualizationManager.cpp
+++ b/Engine/Source/Developer/Virtualization/Private/VirtualizationManager.cpp
@@ -407,7 +407,7 @@ bool FVirtualizationManager::PushData(TArrayView<FPushRequest> Requests, EStorag
 		return false;
 	}
 
-	FConditionalScopeLock _(&ForceSingleThreadedCS, DebugValues.bSingleThreaded);
+	FConditionalScopeLock _(&DebugValues.ForceSingleThreadedCS, DebugValues.bSingleThreaded);
 
 	// TODO: Note that all push operations are currently synchronous, probably 
 	// should change to async at some point, although this makes handling failed
@@ -478,7 +478,7 @@ FCompressedBuffer FVirtualizationManager::PullData(const FIoHash& Id)
 		return FCompressedBuffer();
 	}
 
-	FConditionalScopeLock _(&ForceSingleThreadedCS, DebugValues.bSingleThreaded);
+	FConditionalScopeLock _(&DebugValues.ForceSingleThreadedCS, DebugValues.bSingleThreaded);
 
 	GetNotificationEvent().Broadcast(IVirtualizationSystem::PullBegunNotification, Id);
 
@@ -530,7 +530,7 @@ EQueryResult FVirtualizationManager::QueryPayloadStatuses(TArrayView<const FIoHa
 	Results.SetNum(Ids.Num());
 
 	{
-		FConditionalScopeLock _(&ForceSingleThreadedCS, DebugValues.bSingleThreaded);
+		FConditionalScopeLock _(&DebugValues.ForceSingleThreadedCS, DebugValues.bSingleThreaded);
 
 		for (IVirtualizationBackend* Backend : Backends)
 		{
