Fix a bug where the virtualization process was failing to virtualize loaded packages when submitted from the editor

#rnx

- When adding the check to see if a package file could be written to, I tried to do this as early as possible to avoid virtualizing payloads for packages that can't be modified.
- This didn't take into account submits via the editor, when the package is loaded which means the editor will have the package file locked for edit.
- For now I have gone with the simple fix, we check if we can write to the package once we have called ResetLoader on the package. This means that we will virtualize the payloads then potentially skip removing them from the package file if another process has a file lock.
- This approach will minimize impact to the users for now but we should revisit this in the future to try and reduce the error scope further.

diff --git a/Engine/Source/Developer/Virtualization/Private/PackageSubmissionChecks.cpp b/Engine/Source/Developer/Virtualization/Private/PackageSubmissionChecks.cpp
--- a/Engine/Source/Developer/Virtualization/Private/PackageSubmissionChecks.cpp
+++ b/Engine/Source/Developer/Virtualization/Private/PackageSubmissionChecks.cpp
@@ -337,20 +337,6 @@ void VirtualizePackages(const TArray<FString>& FilesToSubmit, TArray<FText>& Out
 
 				if (!PkgInfo.LocalPayloads.IsEmpty())
 				{	
-					if (!CanWriteToFile(AbsoluteFilePath))
-					{
-						// Technically the package could have local payloads that won't be virtualized due to filtering or min payload sizes and so the
-						// following warning is misleading. This will be solved if we move that evaluation to the point of saving a package.
-						// If not then we probably need to extend QueryPayloadStatuses to test filtering etc as well, then check for potential package
-						// modification after that.
-						// Long term, the stand alone tool should be able to request the UnrealEditor relinquish the lock on the package file so this becomes 
-						// less of a problem.
-						FText Message = FText::Format(LOCTEXT("Virtualization_PkgLocked", "The package file '{0}' has local payloads but is locked for modification and cannot be virtualized, this package will be skipped!"),
-							FText::FromString(PkgInfo.Path.GetDebugName()));
-						UE_LOG(LogVirtualization, Warning, TEXT("%s"), *Message.ToString());
-						continue;
-					}
-
 					TotalPayloadsToCheck += PkgInfo.LocalPayloads.Num();
 
 					PkgInfo.PayloadIndex = AllLocalPayloads.Num();
@@ -550,12 +536,30 @@ void VirtualizePackages(const TArray<FString>& FilesToSubmit, TArray<FText>& Out
 		// having said that, once a package is in PackagesToReplace it should still be safe to submit so maybe we don't need this level of protection?
 
 		// We need to reset the loader of any package that we want to re-save over
-		for (const TPair<FPackagePath, FString>& Iterator : PackagesToReplace)
+		for (int32 Index = 0; Index < PackagesToReplace.Num(); ++Index)
 		{
-			UPackage* Package = FindObjectFast<UPackage>(nullptr, Iterator.Key.GetPackageFName());
+			const TPair<FPackagePath, FString>& Pair = PackagesToReplace[Index];
+
+			UPackage* Package = FindObjectFast<UPackage>(nullptr, Pair.Key.GetPackageFName());
 			if (Package != nullptr)
 			{
-				ResetLoadersForSave(Package, *Iterator.Key.GetLocalFullPath());
+				UE_LOG(LogVirtualization, Verbose, TEXT("Detaching '%s' from disk so that it can be virtualized"), *Pair.Key.GetDebugName());
+				ResetLoadersForSave(Package, *Pair.Key.GetLocalFullPath());
+			}
+
+			if (!CanWriteToFile(Pair.Key.GetLocalFullPath()))
+			{
+				// Technically the package could have local payloads that won't be virtualized due to filtering or min payload sizes and so the
+				// following warning is misleading. This will be solved if we move that evaluation to the point of saving a package.
+				// If not then we probably need to extend QueryPayloadStatuses to test filtering etc as well, then check for potential package
+				// modification after that.
+				// Long term, the stand alone tool should be able to request the UnrealEditor relinquish the lock on the package file so this becomes 
+				// less of a problem.
+				FText Message = FText::Format(LOCTEXT("Virtualization_PkgLocked", "The package file '{0}' has local payloads but is locked for modification and cannot be virtualized, this package will be skipped!"),
+					FText::FromString(Pair.Key.GetDebugName()));
+				UE_LOG(LogVirtualization, Warning, TEXT("%s"), *Message.ToString());
+				
+				PackagesToReplace.RemoveAt(Index--);
 			}
 		}
 
