The source control VA abackend will no longer early out of pushing payloads if a payload batch is found to already be entirely uploaded.

#rnx

- When the code was converted from single batch to support multiple batches I missed the check to return true if all payloads were already uploaded. So once we were operating on multiple batches if a batch was 100% already present in the backend we would return true and not process any further batches.
- Change the return true, to a continue so that future batches are processed.

diff --git a/Engine/Source/Developer/Virtualization/Private/VirtualizationSourceControlBackend.cpp b/Engine/Source/Developer/Virtualization/Private/VirtualizationSourceControlBackend.cpp
--- a/Engine/Source/Developer/Virtualization/Private/VirtualizationSourceControlBackend.cpp
+++ b/Engine/Source/Developer/Virtualization/Private/VirtualizationSourceControlBackend.cpp
@@ -620,7 +620,8 @@ bool FSourceControlBackend::PushData(TArrayView<FPushRequest> Requests)
 
 		if (FilesToAdd.IsEmpty())
 		{
-			return true;
+			// If we have no files to add then we should skip to the next batch
+			continue;
 		}
 
 		{
