Add 'size of disk' to the output for local payloads when running the 'DumpPackagePayloadInfo command.

#rnx

diff --git a/Engine/Source/Developer/Virtualization/Private/DumpPackagePayloadInfoCommand.cpp b/Engine/Source/Developer/Virtualization/Private/DumpPackagePayloadInfoCommand.cpp
--- a/Engine/Source/Developer/Virtualization/Private/DumpPackagePayloadInfoCommand.cpp
+++ b/Engine/Source/Developer/Virtualization/Private/DumpPackagePayloadInfoCommand.cpp
@@ -10,6 +10,24 @@ namespace UE
 
 #if WITH_EDITORONLY_DATA
 
+FString BytesToString(int64 SizeInBytes)
+{
+	if (SizeInBytes < (8 *1024))
+	{
+		return FString::Printf(TEXT("%4d bytes"), SizeInBytes);
+	}
+	else if (SizeInBytes < (1024 * 1024))
+	{
+		double SizeInKb = SizeInBytes / (1024.0);
+		return FString::Printf(TEXT("%.2f KB"), SizeInKb);
+	}
+	else
+	{
+		double SizeInMB = SizeInBytes / (1024.0 * 1024.0);
+		return FString::Printf(TEXT("%.2f MB"), SizeInMB);
+	}
+}
+
 /**
  * This function is used to write information about package's payloads to the log file. This has no 
  * practical development use and should only be used for debugging purposes. 
@@ -47,11 +65,15 @@ void DumpPackagePayloadInfo(const TArray<FString>& Args)
 			if (LocalPayloadIds.Num() > 0)
 			{
 				UE_LOG(LogVirtualization, Display, TEXT("LocalPayloads:"));
-				UE_LOG(LogVirtualization, Display, TEXT("Index|\t%-40s|\tFilterReason"), TEXT("PayloadIdentifier"));
+				UE_LOG(LogVirtualization, Display, TEXT("Index | %-40s | SizeOnDisk | FilterReason"), TEXT("PayloadIdentifier"));
 				for (int32 Index = 0; Index < LocalPayloadIds.Num(); ++Index)
 				{
 					FPayloadInfo Info = Trailer.GetPayloadInfo(LocalPayloadIds[Index]);
-					UE_LOG(LogVirtualization, Display, TEXT("%02d:  |\t%s|\t%s"), Index, *LexToString(LocalPayloadIds[Index]), *LexToString(Info.FilterFlags));
+					UE_LOG(LogVirtualization, Display, TEXT("%02d    | %s | %-10s | %s"), 
+						Index, 
+						*LexToString(LocalPayloadIds[Index]), 
+						*BytesToString(Info.CompressedSize),
+						*LexToString(Info.FilterFlags));
 				}
 			}
 
