The DDC virtualization backend now supports the batch push API

#rnx

- The DDC2 api is already heavily threaded so we can simply keep throwing new requests onto it and wait once all requests are made and not have to do any work ourselves.
- I did ask and was told that using the batch api provided by DDC2 would not improve performance so for now I am have gone with the simple change. We may still look into changing to the batch API in the future as then we would have a reason to improve the DDC2 api.

diff --git a/Engine/Source/Developer/Virtualization/Private/VirtualizationDDCBackend.cpp b/Engine/Source/Developer/Virtualization/Private/VirtualizationDDCBackend.cpp
--- a/Engine/Source/Developer/Virtualization/Private/VirtualizationDDCBackend.cpp
+++ b/Engine/Source/Developer/Virtualization/Private/VirtualizationDDCBackend.cpp
@@ -115,6 +115,63 @@
 	}
 }
 
+bool FDDCBackend::PushData(TArrayView<FPushRequest> Requests)
+{
+	TRACE_CPUPROFILER_EVENT_SCOPE(FDDCBackend::PushData);
+
+	UE::DerivedData::ICache& Cache = UE::DerivedData::GetCache();
+
+	UE::DerivedData::FRequestOwner Owner(UE::DerivedData::EPriority::Normal);
+
+	bool bWasSuccess = true;
+
+	// TODO: We tend not to memory bloat too much on large batches as the requests complete quite quickly
+	// however we might want to consider adding better control on how much total memory we can dedicate to
+	// loading payloads before we wait for requests to complete?
+	for (FPushRequest& Request : Requests)
+	{
+		if (DoesPayloadExist(Request.GetIdentifier()))
+		{
+			Request.SetStatus(FPushRequest::EStatus::Success);
+		}
+		else
+		{
+			UE::DerivedData::FRequestBarrier Barrier(Owner);
+
+			UE::DerivedData::FCacheKey Key;
+			Key.Bucket = Bucket;
+			Key.Hash = Request.GetIdentifier();
+
+			UE::DerivedData::FValue DerivedDataValue(Request.GetPayload());
+			check(DerivedDataValue.GetRawHash() == Request.GetIdentifier());
+
+			UE::DerivedData::FCacheRecordBuilder RecordBuilder(Key);
+			RecordBuilder.AddValue(ToDerivedDataValueId(Request.GetIdentifier()), DerivedDataValue);
+
+			UE::DerivedData::FCachePutResponse Result;
+			auto Callback = [&Request, &bWasSuccess](UE::DerivedData::FCachePutResponse&& Response)
+			{
+				if (Response.Status == UE::DerivedData::EStatus::Ok)
+				{
+					Request.SetStatus(FPushRequest::EStatus::Success);
+				}
+				else
+				{
+					Request.SetStatus(FPushRequest::EStatus::Failed);
+					bWasSuccess = false;
+				}
+			};
+
+			// TODO: Improve the name when we start passing more context to this function
+			Cache.Put({ {{TEXT("Mirage")}, RecordBuilder.Build(), TransferPolicy} }, Owner, MoveTemp(Callback));
+		}
+	}
+
+	Owner.Wait();
+
+	return bWasSuccess;
+}
+
 FCompressedBuffer FDDCBackend::PullData(const FIoHash& Id)
 {
 	TRACE_CPUPROFILER_EVENT_SCOPE(FDDCBackend::PullData);
diff --git a/Engine/Source/Developer/Virtualization/Private/VirtualizationDDCBackend.h b/Engine/Source/Developer/Virtualization/Private/VirtualizationDDCBackend.h
--- a/Engine/Source/Developer/Virtualization/Private/VirtualizationDDCBackend.h
+++ b/Engine/Source/Developer/Virtualization/Private/VirtualizationDDCBackend.h
@@ -37,6 +37,7 @@ private:
 	virtual bool Initialize(const FString& ConfigEntry) override;
 
 	virtual EPushResult PushData(const FIoHash& Id, const FCompressedBuffer& Payload, const FString& PackageContext) override;
+	virtual bool PushData(TArrayView<FPushRequest> Requests) override;
 
 	virtual FCompressedBuffer PullData(const FIoHash& Id) override;
 	
