Asset virtualization graphs no longer require both 'LocalStorageHierarchy' and 'PersistentStorageHierarchy' defining just one entry is enough if that is required.

#rnx

diff --git a/Engine/Source/Developer/Virtualization/Private/VirtualizationManager.cpp b/Engine/Source/Developer/Virtualization/Private/VirtualizationManager.cpp
--- a/Engine/Source/Developer/Virtualization/Private/VirtualizationManager.cpp
+++ b/Engine/Source/Developer/Virtualization/Private/VirtualizationManager.cpp
@@ -1069,23 +1069,25 @@ void FVirtualizationManager::MountBackends(const FConfigFile& ConfigFile)
 void FVirtualizationManager::ParseHierarchy(const FConfigFile& ConfigFile, const TCHAR* GraphName, const TCHAR* HierarchyKey, const FRegistedFactories& FactoryLookupTable, FBackendArray& PushArray)
 {
 	FString HierarchyData;
-	if (!ConfigFile.GetValue(GraphName, HierarchyKey, HierarchyData))
+	if (ConfigFile.GetValue(GraphName, HierarchyKey, HierarchyData))
 	{
-		UE_LOG(LogVirtualization, Fatal, TEXT("Unable to find the '%s' entry for the content virtualization backend graph '%s' [ini=%s]."), HierarchyKey, GraphName, *GEngineIni);
-	}
-
-	if (HierarchyData.IsEmpty())
-	{
-		UE_LOG(LogVirtualization, Fatal, TEXT("The '%s' entry for backend graph '%s' is empty [ini=%s]."), HierarchyKey, GraphName, *GEngineIni);
-	}
+		if (HierarchyData.IsEmpty())
+		{
+			UE_LOG(LogVirtualization, Fatal, TEXT("The '%s' entry for backend graph '%s' is empty [ini=%s]."), HierarchyKey, GraphName, *GEngineIni);
+		}
 
-	const TArray<FString> Entries = ParseEntries(HierarchyData);
+		const TArray<FString> Entries = ParseEntries(HierarchyData);
 
-	UE_LOG(LogVirtualization, Display, TEXT("'%s' has %d backend(s)"), HierarchyKey, Entries.Num());
+		UE_LOG(LogVirtualization, Display, TEXT("'%s' has %d backend(s)"), HierarchyKey, Entries.Num());
 
-	for (const FString& Entry : Entries)
+		for (const FString& Entry : Entries)
+		{
+			CreateBackend(ConfigFile, GraphName, Entry, FactoryLookupTable, PushArray);
+		}	
+	}
+	else
 	{
-		CreateBackend(ConfigFile, GraphName, Entry, FactoryLookupTable, PushArray);
+		UE_LOG(LogVirtualization, Display, TEXT("No entries for '%s' in the content virtualization backend graph '%s' [ini=%s]."), HierarchyKey, GraphName, *GEngineIni);
 	}
 }
 
