Move the ini file settings for FVirtualizationManager from [Core.ContentVirtualization] to [Core.VirtualizationModule]

#rnx

- The main goal is to split the ini file option that selects the virtualization system to use (SystemName) from the options for our virtualization module.
- The old values should continue to be read and used, but if we detect that ini file values are still under [Core.ContentVirtualization] we will warn about it.

diff --git a/Engine/Config/BaseEngine.ini b/Engine/Config/BaseEngine.ini
--- a/Engine/Config/BaseEngine.ini
+++ b/Engine/Config/BaseEngine.ini
@@ -1212,9 +1212,11 @@ DefaultThreadStackSize=1048576
 EnablePackageSidecarSaving=False
 AllowVirtualizationOptOut =False
 
-; Options for the content virtualization system
+; Sets which virtualization system to use, valid options are None, Default (use the Epic provided system) or a third party system
 [Core.ContentVirtualization]
 SystemName=None
+
+[Core.VirtualizationModule]
 EnablePushToBackend=False
 EnableCacheAfterPull=True
 MinPayloadLength=0
diff --git a/Engine/Source/Developer/Virtualization/Private/VirtualizationManager.cpp b/Engine/Source/Developer/Virtualization/Private/VirtualizationManager.cpp
--- a/Engine/Source/Developer/Virtualization/Private/VirtualizationManager.cpp
+++ b/Engine/Source/Developer/Virtualization/Private/VirtualizationManager.cpp
@@ -669,53 +669,61 @@ void FVirtualizationManager::GetPayloadActivityInfo( GetPayloadActivityInfoFuncR
 void FVirtualizationManager::ApplySettingsFromConfigFiles(const FConfigFile& ConfigFile)
 {
 	UE_LOG(LogVirtualization, Display, TEXT("Loading virtualization manager settings from config files..."));
-	
+
+	// Note that all options are doubled up as we are moving the options for this module from "Core.ContentVirtualization"
+	// to it's own specific "Core.VirtualizationModule" section. This duplication can be removed before we ship 5.1
+
 	bool bEnablePayloadPushingFromIni = false;
-	if (ConfigFile.GetBool(TEXT("Core.ContentVirtualization"), TEXT("EnablePushToBackend"), bEnablePayloadPushingFromIni))
+	if (ConfigFile.GetBool(TEXT("Core.ContentVirtualization"), TEXT("EnablePushToBackend"), bEnablePayloadPushingFromIni) ||
+		ConfigFile.GetBool(TEXT("Core.VirtualizationModule"), TEXT("EnablePushToBackend"), bEnablePayloadPushingFromIni))
 	{
 		bEnablePayloadPushing = bEnablePayloadPushingFromIni;
 		UE_LOG(LogVirtualization, Display, TEXT("\tEnablePushToBackend : %s"), bEnablePayloadPushing ? TEXT("true") : TEXT("false") );
 	}
 	else
 	{
-		UE_LOG(LogVirtualization, Error, TEXT("Failed to load [Core.ContentVirtualization].EnablePushToBackend from config file!"));
+		UE_LOG(LogVirtualization, Error, TEXT("Failed to load [Core.VirtualizationModule].EnablePushToBackend from config file!"));
 	}
 
 	bool bEnableCacheAfterPullFromIni = false;
-	if (ConfigFile.GetBool(TEXT("Core.ContentVirtualization"), TEXT("EnableCacheAfterPull"), bEnableCacheAfterPullFromIni))
+	if (ConfigFile.GetBool(TEXT("Core.ContentVirtualization"), TEXT("EnableCacheAfterPull"), bEnableCacheAfterPullFromIni) ||
+		ConfigFile.GetBool(TEXT("Core.VirtualizationModule"), TEXT("EnableCacheAfterPull"), bEnableCacheAfterPullFromIni))
 	{
 		bEnableCacheAfterPull = bEnableCacheAfterPullFromIni;
 		UE_LOG(LogVirtualization, Display, TEXT("\tCachePulledPayloads : %s"), bEnableCacheAfterPull ? TEXT("true") : TEXT("false"));
 	}
 	else
 	{
-		UE_LOG(LogVirtualization, Error, TEXT("Failed to load [Core.ContentVirtualization].EnableCacheAfterPull from config file!"));
+		UE_LOG(LogVirtualization, Error, TEXT("Failed to load [Core.VirtualizationModule].EnableCacheAfterPull from config file!"));
 	}
 
 	int64 MinPayloadLengthFromIni = 0;
-	if (ConfigFile.GetInt64(TEXT("Core.ContentVirtualization"), TEXT("MinPayloadLength"), MinPayloadLengthFromIni))
+	if (ConfigFile.GetInt64(TEXT("Core.ContentVirtualization"), TEXT("MinPayloadLength"), MinPayloadLengthFromIni) ||
+		ConfigFile.GetInt64(TEXT("Core.VirtualizationModule"), TEXT("MinPayloadLength"), MinPayloadLengthFromIni))
 	{
 		MinPayloadLength = MinPayloadLengthFromIni;
 		UE_LOG(LogVirtualization, Display, TEXT("\tMinPayloadLength : %" INT64_FMT), MinPayloadLength );
 	}
 	else
 	{
-		UE_LOG(LogVirtualization, Error, TEXT("Failed to load [Core.ContentVirtualization].MinPayloadLength from config file!"));
+		UE_LOG(LogVirtualization, Error, TEXT("Failed to load [Core.VirtualizationModule].MinPayloadLength from config file!"));
 	}
 
 	FString BackendGraphNameFromIni;
-	if (ConfigFile.GetString(TEXT("Core.ContentVirtualization"), TEXT("BackendGraph"), BackendGraphNameFromIni))
+	if (ConfigFile.GetString(TEXT("Core.ContentVirtualization"), TEXT("BackendGraph"), BackendGraphNameFromIni) ||
+		ConfigFile.GetString(TEXT("Core.VirtualizationModule"), TEXT("BackendGraph"), BackendGraphNameFromIni))
 	{
 		BackendGraphName = BackendGraphNameFromIni;
 		UE_LOG(LogVirtualization, Display, TEXT("\tBackendGraphName : %s"), *BackendGraphName );
 	}
 	else
 	{
-		UE_LOG(LogVirtualization, Error, TEXT("Failed to load [Core.ContentVirtualization].BackendGraph from config file!"));
+		UE_LOG(LogVirtualization, Error, TEXT("Failed to load [Core.VirtualizationModule].BackendGraph from config file!"));
 	}
 
 	FString FilterModeFromIni;
-	if (ConfigFile.GetString(TEXT("Core.ContentVirtualization"), TEXT("FilterMode"), FilterModeFromIni))
+	if (ConfigFile.GetString(TEXT("Core.ContentVirtualization"), TEXT("FilterMode"), FilterModeFromIni) ||
+		ConfigFile.GetString(TEXT("Core.VirtualizationModule"), TEXT("FilterMode"), FilterModeFromIni))
 	{
 		if(LexTryParseString(FilteringMode, FilterModeFromIni))
 		{
@@ -723,38 +731,41 @@ void FVirtualizationManager::ApplySettingsFromConfigFiles(const FConfigFile& Con
 		}
 		else
 		{
-			UE_LOG(LogVirtualization, Error, TEXT("[Core.ContentVirtualization].FilterMode was an invalid value! Allowed: 'OptIn'|'OptOut' Found '%s'"), *FilterModeFromIni);
+			UE_LOG(LogVirtualization, Error, TEXT("[Core.VirtualizationModule].FilterMode was an invalid value! Allowed: 'OptIn'|'OptOut' Found '%s'"), *FilterModeFromIni);
 		}
 	}
 	else
 	{
-		UE_LOG(LogVirtualization, Error, TEXT("Failed to load [Core.ContentVirtualization]FilterMode from config file!"));
+		UE_LOG(LogVirtualization, Error, TEXT("Failed to load [Core.VirtualizationModule]FilterMode from config file!"));
 	}
 
 	bool bFilterEngineContentFromIni = true;
-	if (ConfigFile.GetBool(TEXT("Core.ContentVirtualization"), TEXT("FilterEngineContent"), bFilterEngineContentFromIni))
+	if (ConfigFile.GetBool(TEXT("Core.ContentVirtualization"), TEXT("FilterEngineContent"), bFilterEngineContentFromIni) ||
+		ConfigFile.GetBool(TEXT("Core.VirtualizationModule"), TEXT("FilterEngineContent"), bFilterEngineContentFromIni))
 	{
 		bFilterEngineContent = bFilterEngineContentFromIni;
 		UE_LOG(LogVirtualization, Display, TEXT("\tFilterEngineContent : %s"), bFilterEngineContent ? TEXT("true") : TEXT("false"));
 	}
 	else
 	{
-		UE_LOG(LogVirtualization, Error, TEXT("Failed to load [Core.ContentVirtualization].FilterEngineContent from config file!"));
+		UE_LOG(LogVirtualization, Error, TEXT("Failed to load [Core.VirtualizationModule].FilterEngineContent from config file!"));
 	}
 	
 	bool bFilterEnginePluginContentFromIni = true;
-	if (ConfigFile.GetBool(TEXT("Core.ContentVirtualization"), TEXT("FilterEnginePluginContent"), bFilterEnginePluginContentFromIni))
+	if (ConfigFile.GetBool(TEXT("Core.ContentVirtualization"), TEXT("FilterEnginePluginContent"), bFilterEnginePluginContentFromIni) ||
+		ConfigFile.GetBool(TEXT("Core.VirtualizationModule"), TEXT("FilterEnginePluginContent"), bFilterEnginePluginContentFromIni))
 	{
 		bFilterEnginePluginContent = bFilterEnginePluginContentFromIni;
 		UE_LOG(LogVirtualization, Display, TEXT("\tFilterEnginePluginContent : %s"), bFilterEnginePluginContent ? TEXT("true") : TEXT("false"));
 	}
 	else
 	{
-		UE_LOG(LogVirtualization, Error, TEXT("Failed to load [Core.ContentVirtualization].FilterEnginePluginContent from config file!"));
+		UE_LOG(LogVirtualization, Error, TEXT("Failed to load [Core.VirtualizationModule].FilterEnginePluginContent from config file!"));
 	}
 
 	TArray<FString> DisabledAssetTypesFromIni;
-	if (ConfigFile.GetArray(TEXT("Core.ContentVirtualization"), TEXT("DisabledAsset"), DisabledAssetTypesFromIni) > 0)
+	if (ConfigFile.GetArray(TEXT("Core.ContentVirtualization"), TEXT("DisabledAsset"), DisabledAssetTypesFromIni) > 0 ||
+		ConfigFile.GetArray(TEXT("Core.VirtualizationModule"), TEXT("DisabledAsset"), DisabledAssetTypesFromIni) > 0)
 	{
 		UE_LOG(LogVirtualization, Display, TEXT("\tVirtualization is disabled for payloads of the following assets:"));
 		DisabledAssetTypes.Reserve(DisabledAssetTypesFromIni.Num());
@@ -764,6 +775,26 @@ void FVirtualizationManager::ApplySettingsFromConfigFiles(const FConfigFile& Con
 			DisabledAssetTypes.Add(FName(AssetType));
 		}	
 	}
+
+	// Warnings are supressed until all internal projects are updated
+#if 0
+	// Check for any legacy settings and print them out (easier to do this in one block rather than one and time)
+	if (const FConfigSection* LegacySection = ConfigFile.Find(TEXT("Core.ContentVirtualization")))
+	{
+		if (LegacySection->Num() > 1)
+		{
+			UE_LOG(LogVirtualization, Warning, TEXT("\tFound %d legacy ini file settings under [Core.ContentVirtualization] that should be moved to [Core.VirtualizationModule]"), LegacySection->Num() - 1);
+			for (auto Iterator : *LegacySection)
+			{
+				FString Name = Iterator.Key.ToString();
+				if (Name != TEXT("SystemName"))
+				{
+					UE_LOG(LogVirtualization, Warning, TEXT("\t\t%s"), *Iterator.Key.ToString());
+				}
+			}
+		}
+	}
+#endif
 }
 
 void FVirtualizationManager::ApplyDebugSettingsFromFromCmdline()
