Enable warnings about legacy VA settings in the ini file after the change in 20594207

#rnx

- Now our internal projects are updated we can enable the warning about projects still using the old settings.
- Note that this remains backwards compatible and the legacy settings will be used.

diff --git a/Engine/Source/Developer/Virtualization/Private/VirtualizationManager.cpp b/Engine/Source/Developer/Virtualization/Private/VirtualizationManager.cpp
--- a/Engine/Source/Developer/Virtualization/Private/VirtualizationManager.cpp
+++ b/Engine/Source/Developer/Virtualization/Private/VirtualizationManager.cpp
@@ -776,8 +776,6 @@ void FVirtualizationManager::ApplySettingsFromConfigFiles(const FConfigFile& Con
 		}	
 	}
 
-	// Warnings are supressed until all internal projects are updated
-#if 0
 	// Check for any legacy settings and print them out (easier to do this in one block rather than one and time)
 	if (const FConfigSection* LegacySection = ConfigFile.Find(TEXT("Core.ContentVirtualization")))
 	{
@@ -794,7 +792,6 @@ void FVirtualizationManager::ApplySettingsFromConfigFiles(const FConfigFile& Con
 			}
 		}
 	}
-#endif
 }
 
 void FVirtualizationManager::ApplyDebugSettingsFromFromCmdline()
