Add a FScopedSlowTask to the initialization of the asset virtualization system to make it easier for users to see when the start up time is unreasonable.

#rnx

- Ideally we'd break down the reporting to be per backend but at the moment the slow task system does not let us do that easily because it will only update the text ever 0.2seconds, so if we have a very fast backend followed by a slow one, we'd update the progress bar with the name of the fast backend, then when we update it with the name of the slow backend it will be skipped so the user will still see the name of the fast backend when waiting on the slow one.
- I would rather show less info than risk showing incorrect info.

diff --git a/Engine/Source/Developer/Virtualization/Private/VirtualizationManager.cpp b/Engine/Source/Developer/Virtualization/Private/VirtualizationManager.cpp
--- a/Engine/Source/Developer/Virtualization/Private/VirtualizationManager.cpp
+++ b/Engine/Source/Developer/Virtualization/Private/VirtualizationManager.cpp
@@ -12,6 +12,7 @@
 #include "Misc/PackagePath.h"
 #include "Misc/Parse.h"
 #include "Misc/Paths.h"
+#include "Misc/ScopedSlowTask.h"
 #include "Misc/ScopeLock.h"
 #include "PackageSubmissionChecks.h"
 #include "ProfilingDebugging/CookStats.h"
@@ -293,6 +294,14 @@ bool FVirtualizationManager::Initialize(const FInitParams& InitParams)
 {
 	TRACE_CPUPROFILER_EVENT_SCOPE(FVirtualizationManager::Initialize);
 
+	// TODO: Ideally we'd break this down further, or at least have a FScopedSlowTask for each
+	// backend initialization but the slow task system will only update the UI every 0.2 seconds
+	// so if we have too many small tasks we might show misleading data to the user, so it is 
+	// better for us to have a single scope here at the top level and rely on UnrealInsights for
+	// detailed profiling unless we do something to how FScopedSlowTask updates the UI.
+	FScopedSlowTask SlowTask(1.0f, LOCTEXT("VAInitialize", "Initializing virtualized asset system..."));
+	SlowTask.EnterProgressFrame(1.0f);
+
 	ProjectName = InitParams.ProjectName;
 
 	ApplySettingsFromConfigFiles(InitParams.ConfigFile);
