Temp fix to prevent payloads that should be filtered out from virtualizing from being virtualized if they are already in persistent storage, which was causing unexpected packages to become virtualized.

#rnx

### Problem
- Firstly this fix is a hack and not intended for shipping 5.1, but the real fix is already planned work, which involves moving the filtering from the push process to being applied to each payload when the package is saved, just as the asset filtering is done at the moment.
-- Doing the proper fix will take a bit of time and testing so we need a quicker fix so that I can re-enable the virtualization process to unblock some workflows.
- All new code is wrapped in a single preprocessor define 'ENABLE_FILTERING_HACK' making the code easy to identify and remove.
- Submitting this hack will cause the priority of the real fix to be bumped in the backlog so that we can remove the hack asap.
- The problem code comes from a pass we do on the payloads to check if they are already in persistent storage or not. If the payload is then we just change it to be virtualized and don't attempt to push. In addition to fixing the filtering there is also a work item to investigate if we really want to keep this logic or if we want to remove it anyway. This might be looked into before I look into the filtering fix as an alternative way to remove the hack.

### Fix
- After we have queried the system to check which of the local payloads are already in persistent storage, BUT before we change them over to be virtualized we now run a new pass, checking which payloads should be filtered.
-- This is a bespoke code path expected only to be used here, so in this case it is not in the IVirtualizationSystem API but requires us to cast to FVirtualizationManager to gain access. This should stop anyone else from using the hack before it is removed.
- The filtering check mimics the filtering that would be run on a normal push operation.
- Once done we can check for any payload that would be filtered, and if so, set the status to not found. This prevents the payload from being changed to be virtualized and means it will be included in the normal push operation where it will be correctly filtered.
- None of this is the best way to do it, but it did seem like the easiest way to add the hack in isolation as no existing code was changed, so it should be simple to just 'delete' the hack when ready.

diff --git a/Engine/Source/Developer/Virtualization/Private/PackageSubmissionChecks.cpp b/Engine/Source/Developer/Virtualization/Private/PackageSubmissionChecks.cpp
--- a/Engine/Source/Developer/Virtualization/Private/PackageSubmissionChecks.cpp
+++ b/Engine/Source/Developer/Virtualization/Private/PackageSubmissionChecks.cpp
@@ -17,6 +17,7 @@
 #include "UObject/PackageTrailer.h"
 #include "UObject/UObjectGlobals.h"
 #include "Virtualization/VirtualizationSystem.h"
+#include "VirtualizationManager.h"
 
 #define LOCTEXT_NAMESPACE "Virtualization"
 
@@ -69,7 +70,7 @@ private:
 		const FPayloadData* Data = PayloadLookupTable.Find(Identifier);
 		if (Data == nullptr)
 		{
-			UE_LOG(LogVirtualization, Error, TEXT("FPayloadProvider was unable to find a payload with the identifier '%s'"), 
+			UE_LOG(LogVirtualization, Error, TEXT("FWorkspaceDomainPayloadProvider was unable to find a payload with the identifier '%s'"), 
 				*LexToString(Identifier));
 
 			return FCompressedBuffer();
@@ -79,7 +80,7 @@ private:
 
 		if (!PackageAr.IsValid())
 		{
-			UE_LOG(LogVirtualization, Error, TEXT("FPayloadProvider was unable to open the package '%s' for reading"), 
+			UE_LOG(LogVirtualization, Error, TEXT("FWorkspaceDomainPayloadProvider was unable to open the package '%s' for reading"), 
 				*Data->PackageName);
 
 			return FCompressedBuffer();
@@ -90,7 +91,7 @@ private:
 		FPackageTrailer Trailer;
 		if (!Trailer.TryLoadBackwards(*PackageAr))
 		{
-			UE_LOG(LogVirtualization, Error, TEXT("FPayloadProvider failed to load the package trailer from the package '%s'"), 
+			UE_LOG(LogVirtualization, Error, TEXT("FWorkspaceDomainPayloadProvider failed to load the package trailer from the package '%s'"), 
 				*Data->PackageName);
 
 			return FCompressedBuffer();
@@ -100,7 +101,7 @@ private:
 		
 		if (!Payload)
 		{
-			UE_LOG(LogVirtualization, Error, TEXT("FPayloadProvider was uanble to load the payload '%s' from the package '%s'"),
+			UE_LOG(LogVirtualization, Error, TEXT("FWorkspaceDomainPayloadProvider was uanble to load the payload '%s' from the package '%s'"),
 				*LexToString(Identifier),
 				*Data->PackageName);
 
@@ -109,7 +110,7 @@ private:
 
 		if (Identifier != FIoHash(Payload.GetRawHash()))
 		{
-			UE_LOG(LogVirtualization, Error, TEXT("FPayloadProvider loaded an incorrect payload from the package '%s'. Expected '%s' Loaded  '%s'"), 
+			UE_LOG(LogVirtualization, Error, TEXT("FWorkspaceDomainPayloadProvider loaded an incorrect payload from the package '%s'. Expected '%s' Loaded  '%s'"), 
 				*Data->PackageName,
 				*LexToString(Identifier),
 				*LexToString(Payload.GetRawHash()));
@@ -134,7 +135,7 @@ private:
 		}
 		else
 		{
-			UE_LOG(LogVirtualization, Error, TEXT("FPayloadProvider was unable to find a payload with the identifier '%s'"),
+			UE_LOG(LogVirtualization, Error, TEXT("FWorkspaceDomainPayloadProvider was unable to find a payload with the identifier '%s'"),
 				*LexToString(Identifier));
 
 			return 0;
@@ -158,6 +159,47 @@ private:
 	TMap<FIoHash, FPayloadData> PayloadLookupTable;
 };
 
+#if ENABLE_FILTERING_HACK
+// This filtering provider should only ever be used with FVirtualizationManager::FilterRequests
+// and so does not need to be able to provide the payload, just the payload size.
+class FFilterProvider final : public IPayloadProvider
+{
+public:
+	void RegisterPayload(const FIoHash& PayloadId, uint64 SizeOnDisk)
+	{
+		if (!PayloadId.IsZero())
+		{
+			PayloadLookupTable.Emplace(PayloadId, SizeOnDisk);
+		}
+	}
+
+private:
+	virtual FCompressedBuffer RequestPayload(const FIoHash& Identifier)
+	{
+		checkNoEntry();
+
+		return FCompressedBuffer();
+	}
+
+	virtual uint64 GetPayloadSize(const FIoHash& Identifier)
+	{
+		if (uint64* SizeOnDisk = PayloadLookupTable.Find(Identifier))
+		{
+			return *SizeOnDisk;
+		}
+		else
+		{
+			UE_LOG(LogVirtualization, Error, TEXT("FFilterProvider was unable to find a payload with the identifier '%s'"),
+				*LexToString(Identifier));
+
+			return 0;
+		}
+	}
+
+	TMap<FIoHash, uint64> PayloadLookupTable;
+};
+#endif //ENABLE_FILTERING_HACK
+
 /**
  * Check that the given package ends with PACKAGE_FILE_TAG. Intended to be used to make sure that
  * we have truncated a package correctly when removing the trailers.
@@ -306,6 +348,12 @@ void VirtualizePackages(const TArray<FString>& FilesToSubmit, TArray<FText>& Out
 
 	Progress.EnterProgressFrame(1.0f);
 
+#if ENABLE_FILTERING_HACK
+	FFilterProvider FilterProvider;
+	TArray<Virtualization::FPushRequest> PayloadsToFilter;
+	PayloadsToFilter.Reserve(PayloadsToFilter.Num());
+#endif //ENABLE_FILTERING_HACK
+
 	// From the list of files to submit we need to find all of the valid packages that contain
 	// local payloads that need to be virtualized.
 	int64 TotalPayloadsToCheck = 0;
@@ -337,12 +385,23 @@ void VirtualizePackages(const TArray<FString>& FilesToSubmit, TArray<FText>& Out
 
 				if (!PkgInfo.LocalPayloads.IsEmpty())
 				{	
+#if ENABLE_FILTERING_HACK
+					// Build up an array of push requests that match the order of AllLocalPayloads/PayloadStatuses
+					for (const FIoHash& PayloadId : PkgInfo.LocalPayloads)
+					{
+						const uint64 SizeOnDisk = PkgInfo.Trailer.FindPayloadSizeOnDisk(PayloadId);
+
+						FilterProvider.RegisterPayload(PayloadId, SizeOnDisk);
+						PayloadsToFilter.Emplace(PayloadId, FilterProvider, PkgInfo.Path.GetPackageName());
+					}
+#endif //ENABLE_FILTERING_HACK
+
 					TotalPayloadsToCheck += PkgInfo.LocalPayloads.Num();
 
 					PkgInfo.PayloadIndex = AllLocalPayloads.Num();
 					AllLocalPayloads.Append(PkgInfo.LocalPayloads);
 
-					Packages.Emplace(MoveTemp(PkgInfo));
+					Packages.Emplace(MoveTemp(PkgInfo));		
 				}
 			}
 		}
@@ -361,6 +420,29 @@ void VirtualizePackages(const TArray<FString>& FilesToSubmit, TArray<FText>& Out
 		return;
 	}
 
+#if ENABLE_FILTERING_HACK
+	{
+		check(PayloadStatuses.Num() == PayloadsToFilter.Num());
+		// If VirtualizePackages is running then we know that System is a FVirtualizationManager so we can just cast.
+		// This lets us avoid adding ::FilterRequests to IVirtualizationSystem and keeps the hack contained to this
+		// module.
+		const FVirtualizationManager* Manager = (FVirtualizationManager*)&System;
+		Manager->FilterRequests(PayloadsToFilter);
+
+		// There are many ways we could stop payloads that should be filtered from being auto virtualized if they 
+		// are present in the persistent backend, but the easiest way without changing the existing code paths is
+		// to set the status to NotFound if we know it should be filtered, to make sure that the payload is sent
+		// to the push request where it will be properly rejected by filtering.
+		for (int32 Index = 0; Index < PayloadStatuses.Num(); ++Index)
+		{
+			if (PayloadsToFilter[Index].GetStatus() != FPushRequest::EStatus::Success)
+			{
+				PayloadStatuses[Index] = FPayloadStatus::NotFound;
+			}
+		}
+	}
+#endif
+
 	// Update payloads that are already in persistent storage and don't need to be pushed
 	int64 TotalPayloadsToVirtualize = 0;
 	for (FPackageInfo& PackageInfo : Packages)
diff --git a/Engine/Source/Developer/Virtualization/Private/VirtualizationManager.cpp b/Engine/Source/Developer/Virtualization/Private/VirtualizationManager.cpp
--- a/Engine/Source/Developer/Virtualization/Private/VirtualizationManager.cpp
+++ b/Engine/Source/Developer/Virtualization/Private/VirtualizationManager.cpp
@@ -290,6 +290,37 @@ FVirtualizationManager::~FVirtualizationManager()
 	UE_LOG(LogVirtualization, Log, TEXT("Virtualization manager destroyed"));
 }
 
+#if ENABLE_FILTERING_HACK
+void FVirtualizationManager::FilterRequests(TArray<FPushRequest>& Requests) const
+{
+	// The same filtering code as in FVirtualizationManager::PushData, since this is a hack
+	// I do not want to make any changes to real code paths so it was safer to just duplicate
+	// the checks that we need.
+	for (FPushRequest& Request : Requests)
+	{
+		if (Request.GetIdentifier().IsZero() || Request.GetPayloadSize() == 0)
+		{
+			Request.SetStatus(FPushRequest::EStatus::Invalid);
+			continue;
+		}
+
+		if ((int64)Request.GetPayloadSize() < MinPayloadLength)
+		{
+			Request.SetStatus(FPushRequest::EStatus::BelowMinSize);
+			continue;
+		}
+
+		if (!ShouldVirtualize(Request.GetContext()))
+		{
+			Request.SetStatus(FPushRequest::EStatus::ExcludedByPackagPath);
+			continue;
+		}
+
+		Request.SetStatus(FPushRequest::EStatus::Success);
+	}
+}
+#endif //ENABLE_FILTERING_HACK
+
 bool FVirtualizationManager::Initialize(const FInitParams& InitParams)
 {
 	TRACE_CPUPROFILER_EVENT_SCOPE(FVirtualizationManager::Initialize);
diff --git a/Engine/Source/Developer/Virtualization/Private/VirtualizationManager.h b/Engine/Source/Developer/Virtualization/Private/VirtualizationManager.h
--- a/Engine/Source/Developer/Virtualization/Private/VirtualizationManager.h
+++ b/Engine/Source/Developer/Virtualization/Private/VirtualizationManager.h
@@ -12,6 +12,18 @@
 class IConsoleObject;
 class FOutputDevice;
 
+// When enabled it will be possible to query push requests to see if they are filtered without actually pushing
+// anything. This is due to a bug with the filtering system, where a payload that should be filtered will be
+// virtualized by our process if the payload is already stored in persistent storage. This hack allows us to 
+// find which payloads should be filtered based on size and package path, then prevent them from being automatically
+// switched to being virtualized. Then when they are 'pushed' the usual filtering path will reject them.
+// The plan is to move this filtering to be done when the package trailer is created, on package save, in which
+// case everything will work just like asset filtering and will be much more robust. However we need to prevent 
+// some development branches from accidentally virtualizing packages which should not be virtualized so this hack
+// is a quick way to get the problem fixed and unblock workflows while the better fix is worked on. This hack
+// should be removed before 5.1 ships.
+#define ENABLE_FILTERING_HACK 1
+
 /**
  * Configuring the backend hierarchy
  * 
@@ -105,6 +117,10 @@ public:
 	FVirtualizationManager();
 	virtual ~FVirtualizationManager();
 
+#if ENABLE_FILTERING_HACK
+	void FilterRequests(TArray<FPushRequest>& Requests) const;
+#endif //ENABLE_FILTERING_HACK
+
 private:
 	/* IVirtualizationSystem implementation */
 
