Rename FPayloadStatus to EPayloadStatus to comply with the coding standards.

#rnx

- I think when I was first adding it, the type was originally a struct but was changed to an enum before it was submitted, however I failed to update the name.

diff --git a/Engine/Source/Developer/Virtualization/Private/PackageSubmissionChecks.cpp b/Engine/Source/Developer/Virtualization/Private/PackageSubmissionChecks.cpp
--- a/Engine/Source/Developer/Virtualization/Private/PackageSubmissionChecks.cpp
+++ b/Engine/Source/Developer/Virtualization/Private/PackageSubmissionChecks.cpp
@@ -411,7 +411,7 @@ void VirtualizePackages(const TArray<FString>& FilesToSubmit, TArray<FText>& Out
 
 	Progress.EnterProgressFrame(1.0f);
 
-	TArray<FPayloadStatus> PayloadStatuses;
+	TArray<EPayloadStatus> PayloadStatuses;
 	if (System.QueryPayloadStatuses(AllLocalPayloads, EStorageType::Persistent, PayloadStatuses) != EQueryResult::Success)
 	{
 		FText Message = LOCTEXT("Virtualization_DoesExistFail", "Failed to find the status of the payloads in the packages being submitted");
@@ -437,7 +437,7 @@ void VirtualizePackages(const TArray<FString>& FilesToSubmit, TArray<FText>& Out
 		{
 			if (PayloadsToFilter[Index].GetStatus() != FPushRequest::EStatus::Success)
 			{
-				PayloadStatuses[Index] = FPayloadStatus::NotFound;
+				PayloadStatuses[Index] = EPayloadStatus::NotFound;
 			}
 		}
 	}
@@ -452,7 +452,7 @@ void VirtualizePackages(const TArray<FString>& FilesToSubmit, TArray<FText>& Out
 #if UE_PRECHECK_PAYLOAD_STATUS
 		for (int32 Index = 0; Index < PackageInfo.LocalPayloads.Num(); ++Index)
 		{
-			if (PayloadStatuses[PackageInfo.PayloadIndex + Index] == FPayloadStatus::FoundAll)
+			if (PayloadStatuses[PackageInfo.PayloadIndex + Index] == EPayloadStatus::FoundAll)
 			{
 				if (PackageInfo.Trailer.UpdatePayloadAsVirtualized(PackageInfo.LocalPayloads[Index]))
 				{
diff --git a/Engine/Source/Developer/Virtualization/Private/VirtualizationManager.cpp b/Engine/Source/Developer/Virtualization/Private/VirtualizationManager.cpp
--- a/Engine/Source/Developer/Virtualization/Private/VirtualizationManager.cpp
+++ b/Engine/Source/Developer/Virtualization/Private/VirtualizationManager.cpp
@@ -563,7 +563,7 @@ FCompressedBuffer FVirtualizationManager::PullData(const FIoHash& Id)
 	}
 }
 
-EQueryResult FVirtualizationManager::QueryPayloadStatuses(TArrayView<const FIoHash> Ids, EStorageType StorageType, TArray<FPayloadStatus>& OutStatuses)
+EQueryResult FVirtualizationManager::QueryPayloadStatuses(TArrayView<const FIoHash> Ids, EStorageType StorageType, TArray<EPayloadStatus>& OutStatuses)
 {
 	TRACE_CPUPROFILER_EVENT_SCOPE(FVirtualizationManager::DoPayloadsExist);
 
@@ -576,7 +576,7 @@ EQueryResult FVirtualizationManager::QueryPayloadStatuses(TArrayView<const FIoHa
 
 	for (int32 Index = 0; Index < Ids.Num(); ++Index)
 	{
-		OutStatuses[Index] = Ids[Index].IsZero() ? FPayloadStatus::Invalid : FPayloadStatus::NotFound;
+		OutStatuses[Index] = Ids[Index].IsZero() ? EPayloadStatus::Invalid : EPayloadStatus::NotFound;
 	}
 
 	FBackendArray& Backends = StorageType == EStorageType::Local ? LocalCachableBackends : PersistentStorageBackends;
@@ -616,15 +616,15 @@ EQueryResult FVirtualizationManager::QueryPayloadStatuses(TArrayView<const FIoHa
 		{
 			if (HitCount[Index] == 0)
 			{
-				OutStatuses[Index] = FPayloadStatus::NotFound;
+				OutStatuses[Index] = EPayloadStatus::NotFound;
 			}
 			else if (HitCount[Index] == Backends.Num())
 			{
-				OutStatuses[Index] = FPayloadStatus::FoundAll;
+				OutStatuses[Index] = EPayloadStatus::FoundAll;
 			}
 			else
 			{
-				OutStatuses[Index] = FPayloadStatus::FoundPartial;
+				OutStatuses[Index] = EPayloadStatus::FoundPartial;
 			}
 		}
 	}
diff --git a/Engine/Source/Developer/Virtualization/Private/VirtualizationManager.h b/Engine/Source/Developer/Virtualization/Private/VirtualizationManager.h
--- a/Engine/Source/Developer/Virtualization/Private/VirtualizationManager.h
+++ b/Engine/Source/Developer/Virtualization/Private/VirtualizationManager.h
@@ -136,7 +136,7 @@ private:
 
 	virtual FCompressedBuffer PullData(const FIoHash& Id) override;
 
-	virtual EQueryResult QueryPayloadStatuses(TArrayView<const FIoHash> Ids, EStorageType StorageType, TArray<FPayloadStatus>& OutStatuses) override;
+	virtual EQueryResult QueryPayloadStatuses(TArrayView<const FIoHash> Ids, EStorageType StorageType, TArray<EPayloadStatus>& OutStatuses) override;
 
 	virtual bool TryVirtualizePackages(const TArray<FString>& FilesToVirtualize, TArray<FText>& OutDescriptionTags, TArray<FText>& OutErrors) override;
 
diff --git a/Engine/Source/Runtime/Core/Private/Virtualization/VirtualizationSystem.cpp b/Engine/Source/Runtime/Core/Private/Virtualization/VirtualizationSystem.cpp
--- a/Engine/Source/Runtime/Core/Private/Virtualization/VirtualizationSystem.cpp
+++ b/Engine/Source/Runtime/Core/Private/Virtualization/VirtualizationSystem.cpp
@@ -62,7 +62,7 @@ public:
 		return FCompressedBuffer();
 	}
 
-	virtual EQueryResult QueryPayloadStatuses(TArrayView<const FIoHash> Ids, EStorageType StorageType, TArray<FPayloadStatus>& OutStatuses) override
+	virtual EQueryResult QueryPayloadStatuses(TArrayView<const FIoHash> Ids, EStorageType StorageType, TArray<EPayloadStatus>& OutStatuses) override
 	{
 		OutStatuses.Reset();
 
diff --git a/Engine/Source/Runtime/Core/Public/Virtualization/VirtualizationSystem.h b/Engine/Source/Runtime/Core/Public/Virtualization/VirtualizationSystem.h
--- a/Engine/Source/Runtime/Core/Public/Virtualization/VirtualizationSystem.h
+++ b/Engine/Source/Runtime/Core/Public/Virtualization/VirtualizationSystem.h
@@ -56,7 +56,7 @@ enum class EQueryResult : int8
 };
 
 /** Describes the status of a payload in regards to a backend storage system */
-enum class FPayloadStatus : int8
+enum class EPayloadStatus : int8
 {
 	/** The payload id was not value */
 	Invalid = -1,
@@ -379,10 +379,10 @@ public:
 	 * @return	True if the operation succeeded and the contents of OutStatuses is valid. False if errors were 
 	 * 			encountered in which case the contents of OutStatuses should be ignored.
 	 */
-	virtual EQueryResult QueryPayloadStatuses(TArrayView<const FIoHash> Ids, EStorageType StorageType, TArray<FPayloadStatus>& OutStatuses) = 0;
+	virtual EQueryResult QueryPayloadStatuses(TArrayView<const FIoHash> Ids, EStorageType StorageType, TArray<EPayloadStatus>& OutStatuses) = 0;
 
 	UE_DEPRECATED(5.1, "Call ::QueryPayloadStatuses instead")
-	bool DoPayloadsExist(TArrayView<const FIoHash> Ids, EStorageType StorageType, TArray<FPayloadStatus>& OutStatuses)
+	bool DoPayloadsExist(TArrayView<const FIoHash> Ids, EStorageType StorageType, TArray<EPayloadStatus>& OutStatuses)
 	{
 		return QueryPayloadStatuses(Ids, StorageType, OutStatuses) != EQueryResult::Success;
 	}
