The virtualization system stats are now only printed to the log when the cook commandlet has finished, rather than every time an editor process exits.

#rnx

- Dumping of stats to the log file is now something that can be invoked via the IVirtualizationSystem api, so anything can call it, if desired.
- Removed the OnExit callback from the virtualization manager, we no longer print the stats to the log file on exit, but instead print them on demand.
- Add a call at the end of the cooking commandlet to dump the virtualization stats so that they appear around the same part of the log file as the cooking stats.
- Changed the verbosity from 'log' to 'display' so that we can see the stats more easily in horde.

diff --git a/Engine/Source/Developer/Virtualization/Private/VirtualizationManager.cpp b/Engine/Source/Developer/Virtualization/Private/VirtualizationManager.cpp
--- a/Engine/Source/Developer/Virtualization/Private/VirtualizationManager.cpp
+++ b/Engine/Source/Developer/Virtualization/Private/VirtualizationManager.cpp
@@ -2,18 +2,18 @@
 
 #include "VirtualizationManager.h"
 
+#include "HAL/IConsoleManager.h"
 #include "HAL/PlatformTime.h"
 #include "IVirtualizationBackend.h"
 #include "Logging/MessageLog.h"
 #include "Misc/CommandLine.h"
 #include "Misc/ConfigCacheIni.h"
-#include "Misc/CoreDelegates.h"
 #include "Misc/PackageName.h"
 #include "Misc/PackagePath.h"
 #include "Misc/Parse.h"
 #include "Misc/Paths.h"
-#include "Misc/ScopedSlowTask.h"
 #include "Misc/ScopeLock.h"
+#include "Misc/ScopedSlowTask.h"
 #include "PackageSubmissionChecks.h"
 #include "ProfilingDebugging/CookStats.h"
 #include "VirtualizationFilterSettings.h"
@@ -188,14 +188,14 @@ namespace Profiling
 			return; // Early out if we have no data
 		}
 
-		UE_LOG(LogVirtualization, Log, TEXT(""));
-		UE_LOG(LogVirtualization, Log, TEXT("Virtualization ProfileData"));
-		UE_LOG(LogVirtualization, Log, TEXT("======================================================================================="));
+		UE_LOG(LogVirtualization, Display, TEXT(""));
+		UE_LOG(LogVirtualization, Display, TEXT("Virtualization ProfileData"));
+		UE_LOG(LogVirtualization, Display, TEXT("======================================================================================="));
 
 		if (CacheStats.Num() > 0)
 		{
-			UE_LOG(LogVirtualization, Log, TEXT("%-40s|%17s|%12s|%14s|"), TEXT("Caching Data"), TEXT("TotalSize (MB)"), TEXT("TotalTime(s)"), TEXT("DataRate(MB/S)"));
-			UE_LOG(LogVirtualization, Log, TEXT("----------------------------------------|-----------------|------------|--------------|"));
+			UE_LOG(LogVirtualization, Display, TEXT("%-40s|%17s|%12s|%14s|"), TEXT("Caching Data"), TEXT("TotalSize (MB)"), TEXT("TotalTime(s)"), TEXT("DataRate(MB/S)"));
+			UE_LOG(LogVirtualization, Display, TEXT("----------------------------------------|-----------------|------------|--------------|"));
 
 			for (const auto& Iterator : CacheStats)
 			{
@@ -203,20 +203,20 @@ namespace Profiling
 				const int64 DataSizeMB = Iterator.Value.GetAccumulatedValueAnyThread(FCookStats::CallStats::EHitOrMiss::Hit, FCookStats::CallStats::EStatType::Bytes) / (1024 * 1024);
 				const double MBps = Time != 0.0 ? (DataSizeMB / Time) : 0.0;
 
-				UE_LOG(LogVirtualization, Log, TEXT("%-40.40s|%17" UINT64_FMT "|%12.3f|%14.3f|"),
+				UE_LOG(LogVirtualization, Display, TEXT("%-40.40s|%17" UINT64_FMT "|%12.3f|%14.3f|"),
 					*Iterator.Key,
 					DataSizeMB,
 					Time,
 					MBps);
 			}
 
-			UE_LOG(LogVirtualization, Log, TEXT("======================================================================================="));
+			UE_LOG(LogVirtualization, Display, TEXT("======================================================================================="));
 		}
 
 		if (PushStats.Num() > 0)
 		{
-			UE_LOG(LogVirtualization, Log, TEXT("%-40s|%17s|%12s|%14s|"), TEXT("Pushing Data"), TEXT("TotalSize (MB)"), TEXT("TotalTime(s)"), TEXT("DataRate(MB/S)"));
-			UE_LOG(LogVirtualization, Log, TEXT("----------------------------------------|-----------------|------------|--------------|"));
+			UE_LOG(LogVirtualization, Display, TEXT("%-40s|%17s|%12s|%14s|"), TEXT("Pushing Data"), TEXT("TotalSize (MB)"), TEXT("TotalTime(s)"), TEXT("DataRate(MB/S)"));
+			UE_LOG(LogVirtualization, Display, TEXT("----------------------------------------|-----------------|------------|--------------|"));
 
 			for (const auto& Iterator : PushStats)
 			{
@@ -224,20 +224,20 @@ namespace Profiling
 				const int64 DataSizeMB = Iterator.Value.GetAccumulatedValueAnyThread(FCookStats::CallStats::EHitOrMiss::Hit, FCookStats::CallStats::EStatType::Bytes) / (1024 * 1024);
 				const double MBps = Time != 0.0 ? (DataSizeMB / Time) : 0.0;
 
-				UE_LOG(LogVirtualization, Log, TEXT("%-40.40s|%17" UINT64_FMT "|%12.3f|%14.3f|"),
+				UE_LOG(LogVirtualization, Display, TEXT("%-40.40s|%17" UINT64_FMT "|%12.3f|%14.3f|"),
 					*Iterator.Key,
 					DataSizeMB,
 					Time,
 					MBps);
 			}
 
-			UE_LOG(LogVirtualization, Log, TEXT("======================================================================================="));
+			UE_LOG(LogVirtualization, Display, TEXT("======================================================================================="));
 		}
 
 		if (PullStats.Num() > 0)
 		{
-			UE_LOG(LogVirtualization, Log, TEXT("%-40s|%17s|%12s|%14s|"), TEXT("Pulling Data"), TEXT("TotalSize (MB)"), TEXT("TotalTime(s)"), TEXT("DataRate(MB/S)"));
-			UE_LOG(LogVirtualization, Log, TEXT("----------------------------------------|-----------------|------------|--------------|"));
+			UE_LOG(LogVirtualization, Display, TEXT("%-40s|%17s|%12s|%14s|"), TEXT("Pulling Data"), TEXT("TotalSize (MB)"), TEXT("TotalTime(s)"), TEXT("DataRate(MB/S)"));
+			UE_LOG(LogVirtualization, Display, TEXT("----------------------------------------|-----------------|------------|--------------|"));
 
 			for (const auto& Iterator : PullStats)
 			{
@@ -245,14 +245,14 @@ namespace Profiling
 				const int64 DataSizeMB = Iterator.Value.GetAccumulatedValueAnyThread(FCookStats::CallStats::EHitOrMiss::Hit, FCookStats::CallStats::EStatType::Bytes) / (1024 * 1024);
 				const double MBps = Time != 0.0 ? (DataSizeMB / Time) : 0.0;
 
-				UE_LOG(LogVirtualization, Log, TEXT("%-40.40s|%17" UINT64_FMT "|%12.3f|%14.3f|"),
+				UE_LOG(LogVirtualization, Display, TEXT("%-40.40s|%17" UINT64_FMT "|%12.3f|%14.3f|"),
 					*Iterator.Key,
 					DataSizeMB,
 					Time,
 					MBps);
 			}
 
-			UE_LOG(LogVirtualization, Log, TEXT("======================================================================================="));
+			UE_LOG(LogVirtualization, Display, TEXT("======================================================================================="));
 		}
 	}
 #endif // ENABLE_COOK_STATS
@@ -264,12 +264,6 @@ FVirtualizationManager::FVirtualizationManager()
 	, MinPayloadLength(0)
 	, BackendGraphName(TEXT("ContentVirtualizationBackendGraph_None"))
 {
-	TRACE_CPUPROFILER_EVENT_SCOPE(FVirtualizationManager::FVirtualizationManager);
-
-	// Allows us to log the profiling data on process exit. 
-	// TODO: We should just be able to call the logging in the destructor, but 
-	// we need to fix the startup/shutdown ordering of Mirage first.
-	COOK_STAT(FCoreDelegates::OnExit.AddStatic(Profiling::LogStats));
 }
 
 FVirtualizationManager::~FVirtualizationManager()
@@ -642,6 +636,13 @@ bool FVirtualizationManager::TryVirtualizePackages(const TArray<FString>& FilesT
 	return OutErrors.IsEmpty();
 }
 
+void FVirtualizationManager::DumpStats() const
+{
+#if ENABLE_COOK_STATS
+	Profiling::LogStats();
+#endif // ENABLE_COOK_STATS
+}
+
 FPayloadActivityInfo FVirtualizationManager::GetAccumualtedPayloadActivityInfo() const
 {
 	FPayloadActivityInfo Info;
diff --git a/Engine/Source/Developer/Virtualization/Private/VirtualizationManager.h b/Engine/Source/Developer/Virtualization/Private/VirtualizationManager.h
--- a/Engine/Source/Developer/Virtualization/Private/VirtualizationManager.h
+++ b/Engine/Source/Developer/Virtualization/Private/VirtualizationManager.h
@@ -140,6 +140,8 @@ private:
 
 	virtual bool TryVirtualizePackages(const TArray<FString>& FilesToVirtualize, TArray<FText>& OutDescriptionTags, TArray<FText>& OutErrors) override;
 
+	virtual void DumpStats() const override;
+
 	virtual FPayloadActivityInfo GetAccumualtedPayloadActivityInfo() const override;
 
 	virtual void GetPayloadActivityInfo( GetPayloadActivityInfoFuncRef ) const override;
diff --git a/Engine/Source/Editor/UnrealEd/Private/Commandlets/CookCommandlet.cpp b/Engine/Source/Editor/UnrealEd/Private/Commandlets/CookCommandlet.cpp
--- a/Engine/Source/Editor/UnrealEd/Private/Commandlets/CookCommandlet.cpp
+++ b/Engine/Source/Editor/UnrealEd/Private/Commandlets/CookCommandlet.cpp
@@ -55,6 +55,7 @@ DEFINE_LOG_CATEGORY_STATIC(LogCookCommandlet, Log, All);
 #include "AnalyticsEventAttribute.h"
 #include "IAnalyticsProviderET.h"
 #include "AnalyticsET.h"
+#include "Virtualization/VirtualizationSystem.h"
 
 namespace DetailedCookStats
 {
@@ -394,6 +395,8 @@ namespace DetailedCookStats
 		}
 
 		DumpBuildDependencyTrackerStats();
+
+		UE::Virtualization::IVirtualizationSystem::Get().DumpStats();
 	}
 }
 #endif
diff --git a/Engine/Source/Runtime/Core/Private/Virtualization/VirtualizationSystem.cpp b/Engine/Source/Runtime/Core/Private/Virtualization/VirtualizationSystem.cpp
--- a/Engine/Source/Runtime/Core/Private/Virtualization/VirtualizationSystem.cpp
+++ b/Engine/Source/Runtime/Core/Private/Virtualization/VirtualizationSystem.cpp
@@ -69,10 +69,6 @@ public:
 		return EQueryResult::Failure_NotImplemented;
 	}
 
-	virtual void GetPayloadActivityInfo(GetPayloadActivityInfoFuncRef) const override
-	{
-	}
-
 	virtual bool TryVirtualizePackages(const TArray<FString>& FilesToVirtualize, TArray<FText>& OutDescriptionTags, TArray<FText>& OutErrors) override
 	{
 		OutDescriptionTags.Reset();
@@ -83,6 +79,16 @@ public:
 		return false;
 	}
 
+	virtual void DumpStats() const override
+	{
+		// The null implementation will have no stats and nothing to log
+	}
+
+	virtual void GetPayloadActivityInfo(GetPayloadActivityInfoFuncRef) const override
+	{
+		// The null implementation has no stats and nothing to invoke
+	}
+
 	virtual FPayloadActivityInfo GetAccumualtedPayloadActivityInfo() const override
 	{
 		return FPayloadActivityInfo();
diff --git a/Engine/Source/Runtime/Core/Public/Virtualization/VirtualizationSystem.h b/Engine/Source/Runtime/Core/Public/Virtualization/VirtualizationSystem.h
--- a/Engine/Source/Runtime/Core/Public/Virtualization/VirtualizationSystem.h
+++ b/Engine/Source/Runtime/Core/Public/Virtualization/VirtualizationSystem.h
@@ -405,6 +405,9 @@ public:
 	 */
 	virtual bool TryVirtualizePackages(const TArray<FString>& FilesToVirtualize, TArray<FText>& OutDescriptionTags, TArray<FText>& OutErrors) = 0;
 
+	/** When called the system should write any performance stats that it has been gathering to the log file */
+	virtual void DumpStats() const = 0;
+
 	using GetPayloadActivityInfoFuncRef = TFunctionRef<void(const FString& DebugName, const FString& ConfigName, const FPayloadActivityInfo& PayloadInfo)>;
 
 	/** Access profiling info relating to payload activity per backend. Stats will only be collected if ENABLE_COOK_STATS is enabled.*/
