Improve logging during the virtualization process

#rnx

- We now log how many of the original file list turned out to be packages and how many of those packages had package trailers in addition to how many potential payloads for virtualization we found.
- Recently we ran into the problem where a team was unable to see why their virtualization command didn't virtualize a package without debugging, which uncovered that the packages had not been re-saved as expected and lacked trailers. If this logging had been present we would've seen that none of the packages being considered for virtualization had the packages and could've zeroed in on the problem without the aid of the debugger.

diff --git a/Engine/Source/Developer/Virtualization/Private/PackageSubmissionChecks.cpp b/Engine/Source/Developer/Virtualization/Private/PackageSubmissionChecks.cpp
--- a/Engine/Source/Developer/Virtualization/Private/PackageSubmissionChecks.cpp
+++ b/Engine/Source/Developer/Virtualization/Private/PackageSubmissionChecks.cpp
@@ -356,6 +356,8 @@ void VirtualizePackages(const TArray<FString>& FilesToSubmit, TArray<FText>& Out
 
 	// From the list of files to submit we need to find all of the valid packages that contain
 	// local payloads that need to be virtualized.
+	int64 TotalPackagesFound = 0;
+	int64 TotalPackageTrailersFound = 0;
 	int64 TotalPayloadsToCheck = 0;
 	for (const FString& AbsoluteFilePath : FilesToSubmit)
 	{
@@ -364,9 +366,13 @@ void VirtualizePackages(const TArray<FString>& FilesToSubmit, TArray<FText>& Out
 		// TODO: How to handle text packages?
 		if (FPackageName::IsPackageExtension(PackagePath.GetHeaderExtension()) || FPackageName::IsTextPackageExtension(PackagePath.GetHeaderExtension()))
 		{
+			TotalPackagesFound++;
+
 			FPackageTrailer Trailer;
 			if (FPackageTrailer::TryLoadFromPackage(PackagePath, Trailer))
 			{
+				TotalPackageTrailersFound++;
+
 				// The following is not expected to ever happen, currently we give a user facing error but it generally means that the asset is broken somehow.
 				ensureMsgf(Trailer.GetNumPayloads(EPayloadStorageType::Referenced) == 0, TEXT("Trying to virtualize a package that already contains payload references which the workspace file should not ever contain!"));
 				if (Trailer.GetNumPayloads(EPayloadStorageType::Referenced) > 0)
@@ -407,6 +413,7 @@ void VirtualizePackages(const TArray<FString>& FilesToSubmit, TArray<FText>& Out
 		}
 	}
 
+	UE_LOG(LogVirtualization, Display, TEXT("Found %" INT64_FMT " package(s)9, %" INT64_FMT " of which had payload trailers"), TotalPackagesFound, TotalPackageTrailersFound);
 	UE_LOG(LogVirtualization, Display, TEXT("Found %" INT64_FMT " payload(s) in %d package(s) that need to be examined for virtualization"), TotalPayloadsToCheck, Packages.Num());
 
 	Progress.EnterProgressFrame(1.0f);
