Demote perforce backend initialization errors to warnings until the system is refactored. This will prevent some contractors from seeing editor start up errors if they do not have perforce permissions correctly configured for asset virtualization yet.

#rnx

diff --git a/Engine/Source/Developer/Virtualization/Private/VirtualizationSourceControlBackend.cpp b/Engine/Source/Developer/Virtualization/Private/VirtualizationSourceControlBackend.cpp
--- a/Engine/Source/Developer/Virtualization/Private/VirtualizationSourceControlBackend.cpp
+++ b/Engine/Source/Developer/Virtualization/Private/VirtualizationSourceControlBackend.cpp
@@ -257,7 +257,7 @@ bool FSourceControlBackend::Initialize(const FString& ConfigEntry)
 		}
 
 		FMessageLog Log("LogVirtualization");
-		Log.Error(	FText::Format(LOCTEXT("FailedSourceControlConnection", "Failed to connect to source control backend with the following errors:\n{0}\nThe source control backend had trouble connecting!\nTrying logging in with the 'p4 login' command or by using p4vs/UnrealGameSync."),
+		Log.Warning(	FText::Format(LOCTEXT("FailedSourceControlConnection", "Failed to connect to source control backend with the following errors:\n{0}\nThe source control backend had trouble connecting!\nTrying logging in with the 'p4 login' command or by using p4vs/UnrealGameSync."),
 					Errors.ToText()));
 		
 		OnConnectionError();
@@ -274,7 +274,7 @@ bool FSourceControlBackend::Initialize(const FString& ConfigEntry)
 	{
 		FMessageLog Log("LogVirtualization");
 
-		Log.Error(	FText::Format(LOCTEXT("FailedMetaInfo", "Failed to find 'payload_metainfo.txt' in the depot '{0}'\nThe source control backend will be unable to pull payloads, is your source control  config set up correctly?"),
+		Log.Warning(	FText::Format(LOCTEXT("FailedMetaInfo", "Failed to find 'payload_metainfo.txt' in the depot '{0}'\nThe source control backend will be unable to pull payloads, is your source control  config set up correctly?"),
 					FText::FromString(DepotRoot)));
 		
 		OnConnectionError();
@@ -282,11 +282,11 @@ bool FSourceControlBackend::Initialize(const FString& ConfigEntry)
 	}	
 #else
 	TSharedRef<FDownloadFile, ESPMode::ThreadSafe> DownloadCommand = ISourceControlOperation::Create<FDownloadFile>();
-	if (!SCCProvider->TryToDownloadFileFromBackgroundThread(DownloadCommand, PayloadMetaInfoPath))
+	if (true || !SCCProvider->TryToDownloadFileFromBackgroundThread(DownloadCommand, PayloadMetaInfoPath))
 	{
 		FMessageLog Log("LogVirtualization");
 
-		Log.Error(	FText::Format(LOCTEXT("FailedMetaInfo", "Failed to find 'payload_metainfo.txt' in the depot '{0}'\nThe source control backend will be unable to pull payloads, is your source control  config set up correctly?"),
+		Log.Warning(	FText::Format(LOCTEXT("FailedMetaInfo", "Failed to find 'payload_metainfo.txt' in the depot '{0}'\nThe source control backend will be unable to pull payloads, is your source control  config set up correctly?"),
 					FText::FromString(DepotRoot)));
 		
 		OnConnectionError();
@@ -299,7 +299,7 @@ bool FSourceControlBackend::Initialize(const FString& ConfigEntry)
 	{
 		FMessageLog Log("LogVirtualization");
 
-		Log.Error(	FText::Format(LOCTEXT("FailedMetaInfo", "Failed to find 'payload_metainfo.txt' in the depot '{0}'\nThe source control backend will be unable to pull payloads, is your source control  config set up correctly?"),
+		Log.Warning(	FText::Format(LOCTEXT("FailedMetaInfo", "Failed to find 'payload_metainfo.txt' in the depot '{0}'\nThe source control backend will be unable to pull payloads, is your source control  config set up correctly?"),
 					FText::FromString(DepotRoot)));
 
 		OnConnectionError();
