Improve the virtualization log output

#rnx

- Switched the data sizes from integer to double type to avoid small data sizes not showing up

diff --git a/Engine/Source/Developer/Virtualization/Private/VirtualizationManager.cpp b/Engine/Source/Developer/Virtualization/Private/VirtualizationManager.cpp
--- a/Engine/Source/Developer/Virtualization/Private/VirtualizationManager.cpp
+++ b/Engine/Source/Developer/Virtualization/Private/VirtualizationManager.cpp
@@ -200,10 +200,10 @@ namespace Profiling
 			for (const auto& Iterator : CacheStats)
 			{
 				const double Time = Iterator.Value.GetAccumulatedValueAnyThread(FCookStats::CallStats::EHitOrMiss::Hit, FCookStats::CallStats::EStatType::Cycles) * FPlatformTime::GetSecondsPerCycle();
-				const int64 DataSizeMB = Iterator.Value.GetAccumulatedValueAnyThread(FCookStats::CallStats::EHitOrMiss::Hit, FCookStats::CallStats::EStatType::Bytes) / (1024 * 1024);
+				const double DataSizeMB = Iterator.Value.GetAccumulatedValueAnyThread(FCookStats::CallStats::EHitOrMiss::Hit, FCookStats::CallStats::EStatType::Bytes) / (1024.0f * 1024.0f);
 				const double MBps = Time != 0.0 ? (DataSizeMB / Time) : 0.0;
 
-				UE_LOG(LogVirtualization, Display, TEXT("%-40.40s|%17" UINT64_FMT "|%12.3f|%14.3f|"),
+				UE_LOG(LogVirtualization, Display, TEXT("%-40.40s|%17.1f|%12.3f|%14.3f|"),
 					*Iterator.Key,
 					DataSizeMB,
 					Time,
@@ -221,10 +221,10 @@ namespace Profiling
 			for (const auto& Iterator : PushStats)
 			{
 				const double Time = Iterator.Value.GetAccumulatedValueAnyThread(FCookStats::CallStats::EHitOrMiss::Hit, FCookStats::CallStats::EStatType::Cycles) * FPlatformTime::GetSecondsPerCycle();
-				const int64 DataSizeMB = Iterator.Value.GetAccumulatedValueAnyThread(FCookStats::CallStats::EHitOrMiss::Hit, FCookStats::CallStats::EStatType::Bytes) / (1024 * 1024);
+				const double DataSizeMB = Iterator.Value.GetAccumulatedValueAnyThread(FCookStats::CallStats::EHitOrMiss::Hit, FCookStats::CallStats::EStatType::Bytes) / (1024.0f * 1024.0f);
 				const double MBps = Time != 0.0 ? (DataSizeMB / Time) : 0.0;
 
-				UE_LOG(LogVirtualization, Display, TEXT("%-40.40s|%17" UINT64_FMT "|%12.3f|%14.3f|"),
+				UE_LOG(LogVirtualization, Display, TEXT("%-40.40s|%17.1f|%12.3f|%14.3f|"),
 					*Iterator.Key,
 					DataSizeMB,
 					Time,
@@ -242,10 +242,10 @@ namespace Profiling
 			for (const auto& Iterator : PullStats)
 			{
 				const double Time = Iterator.Value.GetAccumulatedValueAnyThread(FCookStats::CallStats::EHitOrMiss::Hit, FCookStats::CallStats::EStatType::Cycles) * FPlatformTime::GetSecondsPerCycle();
-				const int64 DataSizeMB = Iterator.Value.GetAccumulatedValueAnyThread(FCookStats::CallStats::EHitOrMiss::Hit, FCookStats::CallStats::EStatType::Bytes) / (1024 * 1024);
+				const double DataSizeMB = Iterator.Value.GetAccumulatedValueAnyThread(FCookStats::CallStats::EHitOrMiss::Hit, FCookStats::CallStats::EStatType::Bytes) / (1024.0f * 1024.0f);
 				const double MBps = Time != 0.0 ? (DataSizeMB / Time) : 0.0;
 
-				UE_LOG(LogVirtualization, Display, TEXT("%-40.40s|%17" UINT64_FMT "|%12.3f|%14.3f|"),
+				UE_LOG(LogVirtualization, Display, TEXT("%-40.40s|%17.1f|%12.3f|%14.3f|"),
 					*Iterator.Key,
 					DataSizeMB,
 					Time,
@@ -268,13 +268,13 @@ FVirtualizationManager::FVirtualizationManager()
 
 FVirtualizationManager::~FVirtualizationManager()
 {
-	UE_LOG(LogVirtualization, Log, TEXT("Destroying backends"));
-
 	for (IConsoleObject* ConsoleObject : DebugValues.ConsoleObjects)
 	{
 		IConsoleManager::Get().UnregisterConsoleObject(ConsoleObject);
 	}
 
+	UE_LOG(LogVirtualization, Log, TEXT("Destroying backends"));
+
 	LocalCachableBackends.Empty();
 	PersistentStorageBackends.Empty();
 	PullEnabledBackends.Empty();
