Add an option to the source control backend to supress the warning popups if the initial connection fails.

#rnx

- This is a temp fix to an internal problem, where we have a group of people unable to access the source control payload storage system due to permissions but do not actually need to access it at the moment.
- By supressing the connection failure pop up we can allow them to keep working undistrubed.
- Improving how we handle error notifications like this is already on the VA roadmap, at which point this change can be removed.

diff --git a/Engine/Source/Developer/Virtualization/Private/VirtualizationSourceControlBackend.cpp b/Engine/Source/Developer/Virtualization/Private/VirtualizationSourceControlBackend.cpp
--- a/Engine/Source/Developer/Virtualization/Private/VirtualizationSourceControlBackend.cpp
+++ b/Engine/Source/Developer/Virtualization/Private/VirtualizationSourceControlBackend.cpp
@@ -774,6 +774,20 @@ bool FSourceControlBackend::TryApplySettingsFromConfigFiles(const FString& Confi
 		UE_LOG(LogVirtualization, Log, TEXT("[%s] Will push payloads in batches of up to %d payloads(s) at a time"), *GetDebugName(), MaxBatchCount);
 	}
 
+	// Check to see if connection error notification pop ups should be shown or not
+	{
+		FParse::Bool(*ConfigEntry, TEXT("SuppressNotifications="), bSuppressNotifications);
+
+		if (bSuppressNotifications)
+		{
+			UE_LOG(LogVirtualization, Log, TEXT("[%s] Connection pop up warnings are supressed"), *GetDebugName());
+		}
+		else
+		{
+			UE_LOG(LogVirtualization, Log, TEXT("[%s] Connection pop up warnings will be shown"), *GetDebugName());
+		}
+	}
+
 	if (!FindSubmissionWorkingDir(ConfigEntry))
 	{
 		return false;
@@ -850,7 +864,10 @@ void FSourceControlBackend::OnConnectionError()
 		return false;
 	};
 
-	FTSTicker::GetCoreTicker().AddTicker(FTickerDelegate::CreateLambda(Callback));
+	if (!bSuppressNotifications)
+	{
+		FTSTicker::GetCoreTicker().AddTicker(FTickerDelegate::CreateLambda(Callback));
+	}
 }
 
 UE_REGISTER_VIRTUALIZATION_BACKEND_FACTORY(FSourceControlBackend, SourceControl);
diff --git a/Engine/Source/Developer/Virtualization/Private/VirtualizationSourceControlBackend.h b/Engine/Source/Developer/Virtualization/Private/VirtualizationSourceControlBackend.h
--- a/Engine/Source/Developer/Virtualization/Private/VirtualizationSourceControlBackend.h
+++ b/Engine/Source/Developer/Virtualization/Private/VirtualizationSourceControlBackend.h
@@ -43,6 +43,9 @@ class FSemaphore;
  * BatchCount [int32]			The max number of payloads that can be pushed to source control in a
  *								single submit. If the number of payloads in a request batch exceeds
  *								this size then it will be split into multiple smaller batches. [Default=100]
+ * SuppressNotifications[bool]:	When true the system will not display a pop up notification when a 
+ *								connection error occurs, allowing the user to stay unaware of the error
+ *								unless it actually causes some sort of problem. [Default=false]
  * 
  * Environment Variables:
  * UE-VirtualizationWorkingDir [string]:	This can be set to a valid directory path that the backend
@@ -96,6 +99,10 @@ private:
 	/** Should we try to make the temp client partitioned or not? */
 	bool bUsePartitionedClient = true;
 
+	/** When true, the backend will not raise a pop up notification on connection error */
+	bool bSuppressNotifications = false;
+
+	/** The maximum number of files to send in a single source control operation */
 	int32 MaxBatchCount = 100;
 
 	/** A counted semaphore that will limit the number of concurrent connections that we can make */
