Rename PackageSubmissionChecks.cpp/.h to PackageVirtualizationProcess.cpp/.h as it better matches what it does.

#rnx

- The code in the files isn't really a check anymore but the actual process for virtualizing packages (taking the payloads stored in the package trailer, pushing to the backends then modifying the trailer) updating the name reflects this.
- It also helps match the PackageRehydrationProcess.

diff --git a/Engine/Source/Developer/Virtualization/Private/PackageSubmissionChecks.cpp b/Engine/Source/Developer/Virtualization/Private/PackageSubmissionChecks.cpp
deleted file mode 100644
diff --git a/Engine/Source/Developer/Virtualization/Private/PackageSubmissionChecks.h b/Engine/Source/Developer/Virtualization/Private/PackageSubmissionChecks.h
deleted file mode 100644
diff --git a/Engine/Source/Developer/Virtualization/Private/PackageVirtualizationProcess.cpp b/Engine/Source/Developer/Virtualization/Private/PackageVirtualizationProcess.cpp
new file mode 100644
--- /dev/null
+++ b/Engine/Source/Developer/Virtualization/Private/PackageVirtualizationProcess.cpp
@@ -0,0 +1,554 @@
+// Copyright Epic Games, Inc. All Rights Reserved.
+
+#include "PackageVirtualizationProcess.h"
+
+#include "Containers/UnrealString.h"
+#include "HAL/FileManager.h"
+#include "HAL/PlatformFileManager.h"
+#include "HAL/PlatformTime.h"
+#include "Internationalization/Internationalization.h"
+#include "Misc/PackageName.h"
+#include "Misc/Paths.h"
+#include "Misc/ScopedSlowTask.h"
+#include "PackageUtils.h"
+#include "Serialization/EditorBulkData.h"
+#include "UObject/Linker.h"
+#include "UObject/Package.h"
+#include "UObject/PackageResourceManager.h"
+#include "UObject/PackageTrailer.h"
+#include "UObject/UObjectGlobals.h"
+#include "Virtualization/VirtualizationSystem.h"
+#include "VirtualizationManager.h"
+
+
+#define LOCTEXT_NAMESPACE "Virtualization"
+
+// When enabled we will check the payloads to see if they already exist in the persistent storage
+// backends before trying to push them.
+#define UE_PRECHECK_PAYLOAD_STATUS 1
+
+namespace UE::Virtualization
+{
+
+/** 
+ * Implementation of the IPayloadProvider interface so that payloads can be requested on demand
+ * when they are being virtualized.
+ * 
+ * This implementation is not optimized. If a package holds many payloads that are all virtualized
+ * we will end up loading the same trailer over and over, as well as opening the same package file
+ * for read many times.
+ * 
+ * So far this has shown to be a rounding error compared to the actual cost of virtualization 
+ * and so implementing any level of caching has been left as a future task.
+ */
+class FWorkspaceDomainPayloadProvider final : public IPayloadProvider
+{
+public: 
+	FWorkspaceDomainPayloadProvider() = default;
+	virtual ~FWorkspaceDomainPayloadProvider() = default;
+
+	/** Register the payload with it's trailer and package name so that we can access it later as needed */
+	void RegisterPayload(const FIoHash& PayloadId, uint64 SizeOnDisk, const FString& PackageName)
+	{
+		if (!PayloadId.IsZero())
+		{
+			PayloadLookupTable.Emplace(PayloadId, FPayloadData(SizeOnDisk, PackageName));
+		}
+	}
+
+private:
+	virtual FCompressedBuffer RequestPayload(const FIoHash& Identifier) override
+	{
+		if (Identifier.IsZero())
+		{
+			return FCompressedBuffer();
+		}
+
+		const FPayloadData* Data = PayloadLookupTable.Find(Identifier);
+		if (Data == nullptr)
+		{
+			UE_LOG(LogVirtualization, Error, TEXT("FWorkspaceDomainPayloadProvider was unable to find a payload with the identifier '%s'"), 
+				*LexToString(Identifier));
+
+			return FCompressedBuffer();
+		}
+		
+		TUniquePtr<FArchive> PackageAr = IPackageResourceManager::Get().OpenReadExternalResource(EPackageExternalResource::WorkspaceDomainFile, *Data->PackageName);
+
+		if (!PackageAr.IsValid())
+		{
+			UE_LOG(LogVirtualization, Error, TEXT("FWorkspaceDomainPayloadProvider was unable to open the package '%s' for reading"), 
+				*Data->PackageName);
+
+			return FCompressedBuffer();
+		}
+			
+		PackageAr->Seek(PackageAr->TotalSize());
+
+		FPackageTrailer Trailer;
+		if (!Trailer.TryLoadBackwards(*PackageAr))
+		{
+			UE_LOG(LogVirtualization, Error, TEXT("FWorkspaceDomainPayloadProvider failed to load the package trailer from the package '%s'"), 
+				*Data->PackageName);
+
+			return FCompressedBuffer();
+		}
+
+		FCompressedBuffer Payload = Trailer.LoadLocalPayload(Identifier, *PackageAr);
+		
+		if (!Payload)
+		{
+			UE_LOG(LogVirtualization, Error, TEXT("FWorkspaceDomainPayloadProvider was uanble to load the payload '%s' from the package '%s'"),
+				*LexToString(Identifier),
+				*Data->PackageName);
+
+			return FCompressedBuffer();
+		}
+
+		if (Identifier != FIoHash(Payload.GetRawHash()))
+		{
+			UE_LOG(LogVirtualization, Error, TEXT("FWorkspaceDomainPayloadProvider loaded an incorrect payload from the package '%s'. Expected '%s' Loaded  '%s'"), 
+				*Data->PackageName,
+				*LexToString(Identifier),
+				*LexToString(Payload.GetRawHash()));
+
+			return FCompressedBuffer();
+		}
+
+		return Payload;
+	}
+
+	virtual uint64 GetPayloadSize(const FIoHash& Identifier) override
+	{
+		if (Identifier.IsZero())
+		{
+			return 0;
+		}
+
+		const FPayloadData* Data = PayloadLookupTable.Find(Identifier);
+		if (Data != nullptr)
+		{
+			return Data->SizeOnDisk;
+		}
+		else
+		{
+			UE_LOG(LogVirtualization, Error, TEXT("FWorkspaceDomainPayloadProvider was unable to find a payload with the identifier '%s'"),
+				*LexToString(Identifier));
+
+			return 0;
+		}
+	}
+
+	/* This structure holds additional info about the payload that we might need later */
+	struct FPayloadData
+	{
+		FPayloadData(uint64 InSizeOnDisk, const FString& InPackageName)
+			: SizeOnDisk(InSizeOnDisk)
+			, PackageName(InPackageName)
+		{
+
+		}
+
+		uint64 SizeOnDisk;
+		FString PackageName;
+	};
+
+	TMap<FIoHash, FPayloadData> PayloadLookupTable;
+};
+
+#if ENABLE_FILTERING_HACK
+// This filtering provider should only ever be used with FVirtualizationManager::FilterRequests
+// and so does not need to be able to provide the payload, just the payload size.
+class FFilterProvider final : public IPayloadProvider
+{
+public:
+	void RegisterPayload(const FIoHash& PayloadId, uint64 SizeOnDisk)
+	{
+		if (!PayloadId.IsZero())
+		{
+			PayloadLookupTable.Emplace(PayloadId, SizeOnDisk);
+		}
+	}
+
+private:
+	virtual FCompressedBuffer RequestPayload(const FIoHash& Identifier)
+	{
+		checkNoEntry();
+
+		return FCompressedBuffer();
+	}
+
+	virtual uint64 GetPayloadSize(const FIoHash& Identifier)
+	{
+		if (uint64* SizeOnDisk = PayloadLookupTable.Find(Identifier))
+		{
+			return *SizeOnDisk;
+		}
+		else
+		{
+			UE_LOG(LogVirtualization, Error, TEXT("FFilterProvider was unable to find a payload with the identifier '%s'"),
+				*LexToString(Identifier));
+
+			return 0;
+		}
+	}
+
+	TMap<FIoHash, uint64> PayloadLookupTable;
+};
+#endif //ENABLE_FILTERING_HACK
+
+void VirtualizePackages(const TArray<FString>& FilesToSubmit, TArray<FText>& OutDescriptionTags, TArray<FText>& OutErrors)
+{
+	TRACE_CPUPROFILER_EVENT_SCOPE(UE::Virtualization::VirtualizePackages);
+
+	IVirtualizationSystem& System = IVirtualizationSystem::Get();
+
+	// TODO: We could check to see if the package is virtualized even if it is disabled for the project
+	// as a safety feature?
+	if (!System.IsEnabled())
+	{
+		return;
+	}
+
+	if (!System.IsPushingEnabled(EStorageType::Persistent))
+	{
+		UE_LOG(LogVirtualization, Verbose, TEXT("Pushing to persistent backend storage is disabled"));
+		return;
+	}
+
+	const double StartTime = FPlatformTime::Seconds();
+
+	FScopedSlowTask Progress(5.0f, LOCTEXT("Virtualization_Task", "Virtualizing Assets..."));
+	Progress.MakeDialog();
+
+	// Other systems may have added errors to this array, we need to check so later we can determine if this function added any additional errors.
+	const int32 NumErrors = OutErrors.Num();
+
+	struct FPackageInfo
+	{
+		FPackagePath Path;
+		FPackageTrailer Trailer;
+
+		TArray<FIoHash> LocalPayloads;
+		int32 PayloadIndex = INDEX_NONE;
+
+		bool bWasTrailerUpdated = false;
+	};
+
+	UE_LOG(LogVirtualization, Display, TEXT("Considering %d file(s) for virtualization"), FilesToSubmit.Num());
+
+	TArray<FPackageInfo> Packages;
+	Packages.Reserve(FilesToSubmit.Num());
+
+	TArray<FIoHash> AllLocalPayloads;
+	AllLocalPayloads.Reserve(FilesToSubmit.Num());
+
+	Progress.EnterProgressFrame(1.0f);
+
+#if ENABLE_FILTERING_HACK
+	FFilterProvider FilterProvider;
+	TArray<Virtualization::FPushRequest> PayloadsToFilter;
+	PayloadsToFilter.Reserve(PayloadsToFilter.Num());
+#endif //ENABLE_FILTERING_HACK
+
+	// From the list of files to submit we need to find all of the valid packages that contain
+	// local payloads that need to be virtualized.
+	int64 TotalPackagesFound = 0;
+	int64 TotalPackageTrailersFound = 0;
+	int64 TotalPayloadsToCheck = 0;
+	for (const FString& AbsoluteFilePath : FilesToSubmit)
+	{
+		FPackagePath PackagePath = FPackagePath::FromLocalPath(AbsoluteFilePath);
+
+		// TODO: How to handle text packages?
+		if (FPackageName::IsPackageExtension(PackagePath.GetHeaderExtension()) || FPackageName::IsTextPackageExtension(PackagePath.GetHeaderExtension()))
+		{
+			TotalPackagesFound++;
+
+			FPackageTrailer Trailer;
+			if (FPackageTrailer::TryLoadFromPackage(PackagePath, Trailer))
+			{
+				TotalPackageTrailersFound++;
+
+				// The following is not expected to ever happen, currently we give a user facing error but it generally means that the asset is broken somehow.
+				ensureMsgf(Trailer.GetNumPayloads(EPayloadStorageType::Referenced) == 0, TEXT("Trying to virtualize a package that already contains payload references which the workspace file should not ever contain!"));
+				if (Trailer.GetNumPayloads(EPayloadStorageType::Referenced) > 0)
+				{
+					FText Message = FText::Format(	LOCTEXT("Virtualization_PkgHasReferences", "Cannot virtualize the package '{1}' as it has referenced payloads in the trailer"),
+													FText::FromString(PackagePath.GetDebugName()));
+					OutErrors.Add(Message);
+					return;
+				}
+
+				FPackageInfo PkgInfo;
+
+				PkgInfo.Path = MoveTemp(PackagePath);
+				PkgInfo.Trailer = MoveTemp(Trailer);
+				PkgInfo.LocalPayloads = PkgInfo.Trailer.GetPayloads(EPayloadFilter::CanVirtualize);
+
+				if (!PkgInfo.LocalPayloads.IsEmpty())
+				{	
+#if ENABLE_FILTERING_HACK
+					// Build up an array of push requests that match the order of AllLocalPayloads/PayloadStatuses
+					for (const FIoHash& PayloadId : PkgInfo.LocalPayloads)
+					{
+						const uint64 SizeOnDisk = PkgInfo.Trailer.FindPayloadSizeOnDisk(PayloadId);
+
+						FilterProvider.RegisterPayload(PayloadId, SizeOnDisk);
+						PayloadsToFilter.Emplace(PayloadId, FilterProvider, PkgInfo.Path.GetPackageName());
+					}
+#endif //ENABLE_FILTERING_HACK
+
+					TotalPayloadsToCheck += PkgInfo.LocalPayloads.Num();
+
+					PkgInfo.PayloadIndex = AllLocalPayloads.Num();
+					AllLocalPayloads.Append(PkgInfo.LocalPayloads);
+
+					Packages.Emplace(MoveTemp(PkgInfo));		
+				}
+			}
+		}
+	}
+
+	UE_LOG(LogVirtualization, Display, TEXT("Found %" INT64_FMT " package(s), %" INT64_FMT " of which had payload trailers"), TotalPackagesFound, TotalPackageTrailersFound);
+	UE_LOG(LogVirtualization, Display, TEXT("Found %" INT64_FMT " payload(s) in %d package(s) that need to be examined for virtualization"), TotalPayloadsToCheck, Packages.Num());
+
+	Progress.EnterProgressFrame(1.0f);
+
+	TArray<EPayloadStatus> PayloadStatuses;
+	if (System.QueryPayloadStatuses(AllLocalPayloads, EStorageType::Persistent, PayloadStatuses) != EQueryResult::Success)
+	{
+		FText Message = LOCTEXT("Virtualization_DoesExistFail", "Failed to find the status of the payloads in the packages being submitted");
+		OutErrors.Add(Message);
+
+		return;
+	}
+
+#if ENABLE_FILTERING_HACK
+	{
+		check(PayloadStatuses.Num() == PayloadsToFilter.Num());
+		// If VirtualizePackages is running then we know that System is a FVirtualizationManager so we can just cast.
+		// This lets us avoid adding ::FilterRequests to IVirtualizationSystem and keeps the hack contained to this
+		// module.
+		const FVirtualizationManager* Manager = (FVirtualizationManager*)&System;
+		Manager->FilterRequests(PayloadsToFilter);
+
+		// There are many ways we could stop payloads that should be filtered from being auto virtualized if they 
+		// are present in the persistent backend, but the easiest way without changing the existing code paths is
+		// to set the status to NotFound if we know it should be filtered, to make sure that the payload is sent
+		// to the push request where it will be properly rejected by filtering.
+		for (int32 Index = 0; Index < PayloadStatuses.Num(); ++Index)
+		{
+			if (PayloadsToFilter[Index].GetStatus() != FPushRequest::EStatus::Success)
+			{
+				PayloadStatuses[Index] = EPayloadStatus::NotFound;
+			}
+		}
+	}
+#endif
+
+	// Update payloads that are already in persistent storage and don't need to be pushed
+	int64 TotalPayloadsToVirtualize = 0;
+	for (FPackageInfo& PackageInfo : Packages)
+	{
+		check(PackageInfo.LocalPayloads.IsEmpty() || PackageInfo.PayloadIndex != INDEX_NONE); // If we have payloads we should have an index
+
+#if UE_PRECHECK_PAYLOAD_STATUS
+		for (int32 Index = 0; Index < PackageInfo.LocalPayloads.Num(); ++Index)
+		{
+			if (PayloadStatuses[PackageInfo.PayloadIndex + Index] == EPayloadStatus::FoundAll)
+			{
+				if (PackageInfo.Trailer.UpdatePayloadAsVirtualized(PackageInfo.LocalPayloads[Index]))
+				{
+					PackageInfo.bWasTrailerUpdated = true;
+				}
+				else
+				{
+					FText Message = FText::Format(	LOCTEXT("Virtualization_UpdateStatusFailed", "Unable to update the status for the payload '{0}' in the package '{1}'"),
+													FText::FromString(LexToString(PackageInfo.LocalPayloads[Index])),
+													FText::FromString(PackageInfo.Path.GetDebugName()));
+					OutErrors.Add(Message);
+					return;
+				}
+			}
+		}
+
+		// If we made changes we should recalculate the local payloads left
+		if (PackageInfo.bWasTrailerUpdated)
+		{
+			PackageInfo.LocalPayloads = PackageInfo.Trailer.GetPayloads(EPayloadStorageType::Local);
+		}
+#endif
+
+		PackageInfo.PayloadIndex = INDEX_NONE;
+		TotalPayloadsToVirtualize += PackageInfo.LocalPayloads.Num();
+	}
+
+	UE_LOG(LogVirtualization, Display, TEXT("Found %" INT64_FMT " payload(s) that potentially need to be pushed to persistent virtualized storage"), TotalPayloadsToVirtualize);
+
+	// TODO Optimization: In theory we could have many packages sharing the same payload and we only need to push once
+	Progress.EnterProgressFrame(1.0f);
+
+	// Build up the info in the payload provider and the final array of payload push requests
+
+	FWorkspaceDomainPayloadProvider PayloadProvider;
+	TArray<Virtualization::FPushRequest> PayloadsToSubmit;
+	PayloadsToSubmit.Reserve(TotalPayloadsToVirtualize);
+
+	for (FPackageInfo& PackageInfo : Packages)
+	{
+		if (PackageInfo.LocalPayloads.IsEmpty())
+		{
+			continue;
+		}
+
+		PackageInfo.PayloadIndex = PayloadsToSubmit.Num();
+
+		for (const FIoHash& PayloadId : PackageInfo.LocalPayloads)
+		{
+			const uint64 SizeOnDisk = PackageInfo.Trailer.FindPayloadSizeOnDisk(PayloadId);
+
+			PayloadProvider.RegisterPayload(PayloadId, SizeOnDisk, PackageInfo.Path.GetPackageName());
+			PayloadsToSubmit.Emplace(PayloadId, PayloadProvider, PackageInfo.Path.GetPackageName());
+		}
+	}
+
+	Progress.EnterProgressFrame(1.0f);
+	// Push any remaining local payload to the persistent backends
+	if (!System.PushData(PayloadsToSubmit, EStorageType::Persistent))
+	{
+		FText Message = LOCTEXT("Virtualization_PushFailure", "Failed to push payloads");
+		OutErrors.Add(Message);
+		return;
+	}
+
+	int64 TotalPayloadsVirtualized = 0;
+	for (const Virtualization::FPushRequest& Request : PayloadsToSubmit)
+	{
+		TotalPayloadsVirtualized += Request.GetStatus() == FPushRequest::EStatus::Success ? 1 : 0;
+	}
+	UE_LOG(LogVirtualization, Display, TEXT("Pushed %" INT64_FMT " payload(s) to persistent virtualized storage"), TotalPayloadsVirtualized);
+
+	// Update the package info for the submitted payloads
+	for (FPackageInfo& PackageInfo : Packages)
+	{
+		for (int32 Index = 0; Index < PackageInfo.LocalPayloads.Num(); ++Index)
+		{
+			const Virtualization::FPushRequest& Request = PayloadsToSubmit[PackageInfo.PayloadIndex + Index];
+			check(Request.GetIdentifier() == PackageInfo.LocalPayloads[Index]);
+
+			if (Request.GetStatus() == Virtualization::FPushRequest::EStatus::Success)
+			{
+				if (PackageInfo.Trailer.UpdatePayloadAsVirtualized(Request.GetIdentifier()))
+				{
+					PackageInfo.bWasTrailerUpdated = true;
+				}
+				else
+				{
+					FText Message = FText::Format(	LOCTEXT("Virtualization_UpdateStatusFailed", "Unable to update the status for the payload '{0}' in the package '{1}'"),
+													FText::FromString(LexToString(Request.GetIdentifier())),
+													FText::FromString(PackageInfo.Path.GetDebugName()));
+					OutErrors.Add(Message);
+					return;
+				}
+			}
+		}
+	}
+
+	Progress.EnterProgressFrame(1.0f);
+
+	TArray<TPair<FPackagePath, FString>> PackagesToReplace;
+
+	// Any package with an updated trailer needs to be copied and an updated trailer appended
+	for (FPackageInfo& PackageInfo : Packages)
+	{
+		if (!PackageInfo.bWasTrailerUpdated)
+		{
+			continue;
+		}
+
+		const FPackagePath& PackagePath = PackageInfo.Path; // No need to validate path, we checked this earlier
+
+		FString NewPackagePath = DuplicatePackageWithUpdatedTrailer(PackagePath.GetLocalFullPath(), PackageInfo.Trailer, OutErrors);
+
+		if (!NewPackagePath.IsEmpty())
+		{
+			// Now that we have successfully created a new version of the package with an updated trailer 
+			// we need to mark that it should replace the original package.
+			PackagesToReplace.Emplace(PackagePath, MoveTemp(NewPackagePath));
+		}
+		else
+		{
+			return;
+		}
+
+	}
+
+	UE_LOG(LogVirtualization, Display, TEXT("%d package(s) had their trailer container modified and need to be updated"), PackagesToReplace.Num());
+
+	if (NumErrors == OutErrors.Num())
+	{
+		// TODO: Consider using the SavePackage model (move the original, then replace, so we can restore all of the original packages if needed)
+		// having said that, once a package is in PackagesToReplace it should still be safe to submit so maybe we don't need this level of protection?
+
+		// We need to reset the loader of any package that we want to re-save over
+		for (int32 Index = 0; Index < PackagesToReplace.Num(); ++Index)
+		{
+			const TPair<FPackagePath, FString>& Pair = PackagesToReplace[Index];
+
+			UPackage* Package = FindObjectFast<UPackage>(nullptr, Pair.Key.GetPackageFName());
+			if (Package != nullptr)
+			{
+				UE_LOG(LogVirtualization, Verbose, TEXT("Detaching '%s' from disk so that it can be virtualized"), *Pair.Key.GetDebugName());
+				ResetLoadersForSave(Package, *Pair.Key.GetLocalFullPath());
+			}
+
+			if (!CanWriteToFile(Pair.Key.GetLocalFullPath()))
+			{
+				// Technically the package could have local payloads that won't be virtualized due to filtering or min payload sizes and so the
+				// following warning is misleading. This will be solved if we move that evaluation to the point of saving a package.
+				// If not then we probably need to extend QueryPayloadStatuses to test filtering etc as well, then check for potential package
+				// modification after that.
+				// Long term, the stand alone tool should be able to request the UnrealEditor relinquish the lock on the package file so this becomes 
+				// less of a problem.
+				FText Message = FText::Format(LOCTEXT("Virtualization_PkgLocked", "The package file '{0}' has local payloads but is locked for modification and cannot be virtualized, this package will be skipped!"),
+					FText::FromString(Pair.Key.GetDebugName()));
+				UE_LOG(LogVirtualization, Warning, TEXT("%s"), *Message.ToString());
+				
+				PackagesToReplace.RemoveAt(Index--);
+			}
+		}
+
+		// Since we had no errors we can now replace all of the packages that were virtualized data with the virtualized replacement file.
+		for(const TPair<FPackagePath,FString>&  Iterator : PackagesToReplace)
+		{
+			const FString OriginalPackagePath = Iterator.Key.GetLocalFullPath();
+			const FString& NewPackagePath = Iterator.Value;
+
+			if (!IFileManager::Get().Move(*OriginalPackagePath, *NewPackagePath))
+			{
+				FText Message = FText::Format(	LOCTEXT("Virtualization_MoveFailed", "Unable to replace the package '{0}' with the virtualized version"),
+												FText::FromString(Iterator.Key.GetDebugName()));
+				OutErrors.Add(Message);
+				continue;
+			}
+		}
+	}
+
+	// If we had no new errors add the validation tag to indicate that the packages are safe for submission. 
+	// TODO: Currently this is a simple tag to make it easier for us to track which assets were submitted via the
+	// virtualization process in a test project. This should be expanded when we add proper p4 server triggers.
+	if (NumErrors == OutErrors.Num())
+	{
+		FText Tag = FText::FromString(TEXT("#virtualized"));
+		OutDescriptionTags.Add(Tag);
+	}
+
+	const double TimeInSeconds = FPlatformTime::Seconds() - StartTime;
+	UE_LOG(LogVirtualization, Verbose, TEXT("Virtualization pre submit check took %.3f(s)"), TimeInSeconds);
+}
+
+} // namespace UE::Virtualization
+
+#undef LOCTEXT_NAMESPACE
+
diff --git a/Engine/Source/Developer/Virtualization/Private/PackageVirtualizationProcess.h b/Engine/Source/Developer/Virtualization/Private/PackageVirtualizationProcess.h
new file mode 100644
--- /dev/null
+++ b/Engine/Source/Developer/Virtualization/Private/PackageVirtualizationProcess.h
@@ -0,0 +1,15 @@
+// Copyright Epic Games, Inc. All Rights Reserved.
+
+#pragma once
+
+#include "Containers/Array.h"
+#include "Containers/UnrealString.h"
+#include "Internationalization/Text.h"
+
+namespace UE::Virtualization
+{
+
+/** @See IVirtualizationSystem::TryVirtualizePackages */
+void VirtualizePackages(const TArray<FString>& FilesToSubmit, TArray<FText>& OutDescriptionTags, TArray<FText>& OutErrors);
+
+} // namespace UE::Virtualization
diff --git a/Engine/Source/Developer/Virtualization/Private/Virtualization.cpp b/Engine/Source/Developer/Virtualization/Private/Virtualization.cpp
--- a/Engine/Source/Developer/Virtualization/Private/Virtualization.cpp
+++ b/Engine/Source/Developer/Virtualization/Private/Virtualization.cpp
@@ -4,9 +4,6 @@
 
 #include "Features/IModularFeatures.h"
 #include "MessageLogModule.h"
-#include "Misc/DelayedAutoRegister.h"
-#include "PackageSubmissionChecks.h"
-#include "Serialization/EditorBulkData.h"
 #include "VirtualizationSourceControlUtilities.h"
 
 #define LOCTEXT_NAMESPACE "Virtualization"
diff --git a/Engine/Source/Developer/Virtualization/Private/VirtualizationManager.cpp b/Engine/Source/Developer/Virtualization/Private/VirtualizationManager.cpp
--- a/Engine/Source/Developer/Virtualization/Private/VirtualizationManager.cpp
+++ b/Engine/Source/Developer/Virtualization/Private/VirtualizationManager.cpp
@@ -15,7 +15,7 @@
 #include "Misc/ScopeLock.h"
 #include "PackageRehydrationProcess.h"
 #include "Misc/ScopedSlowTask.h"
-#include "PackageSubmissionChecks.h"
+#include "PackageVirtualizationProcess.h"
 #include "ProfilingDebugging/CookStats.h"
 #include "VirtualizationFilterSettings.h"
 
