The depot path provided to the VA source control backend no longer requires a trailing /

#rnx

- When the depot path is parsed from the config file we now add the trailing '/' if it is not already there, this allows us to make the input more forgiving without changing the existing code.
- Removed FString::Printf when creating the path for payload_metainfo in favour of using WriteToString, this doesn't add much but gets rid of the %s.
- Reformatted some messages that had odd tab alignment.

diff --git a/Engine/Source/Developer/Virtualization/Private/VirtualizationSourceControlBackend.cpp b/Engine/Source/Developer/Virtualization/Private/VirtualizationSourceControlBackend.cpp
--- a/Engine/Source/Developer/Virtualization/Private/VirtualizationSourceControlBackend.cpp
+++ b/Engine/Source/Developer/Virtualization/Private/VirtualizationSourceControlBackend.cpp
@@ -266,7 +266,7 @@ bool FSourceControlBackend::Initialize(const FString& ConfigEntry)
 
 	// When a source control depot is set up a file named 'payload_metainfo.txt' should be submitted to it's root.
 	// This allows us to check for the existence of the file to confirm that the depot root is indeed valid.
-	const FString PayloadMetaInfoPath = FString::Printf(TEXT("%spayload_metainfo.txt"), *DepotRoot);
+	const FString PayloadMetaInfoPath = WriteToString<512>(DepotRoot, TEXT("payload_metainfo.txt")).ToString();
 
 #if IS_SOURCE_CONTROL_THREAD_SAFE
 	TSharedRef<FDownloadFile, ESPMode::ThreadSafe> DownloadCommand = ISourceControlOperation::Create<FDownloadFile>();
@@ -274,7 +274,7 @@ bool FSourceControlBackend::Initialize(const FString& ConfigEntry)
 	{
 		FMessageLog Log("LogVirtualization");
 
-		Log.Warning(	FText::Format(LOCTEXT("FailedMetaInfo", "Failed to find 'payload_metainfo.txt' in the depot '{0}'\nThe source control backend will be unable to pull payloads, is your source control  config set up correctly?"),
+		Log.Warning(FText::Format(LOCTEXT("FailedMetaInfo", "Failed to find 'payload_metainfo.txt' in the depot '{0}'\nThe source control backend will be unable to pull payloads, is your source control  config set up correctly?"),
 					FText::FromString(DepotRoot)));
 		
 		OnConnectionError();
@@ -286,7 +286,7 @@ bool FSourceControlBackend::Initialize(const FString& ConfigEntry)
 	{
 		FMessageLog Log("LogVirtualization");
 
-		Log.Warning(	FText::Format(LOCTEXT("FailedMetaInfo", "Failed to find 'payload_metainfo.txt' in the depot '{0}'\nThe source control backend will be unable to pull payloads, is your source control  config set up correctly?"),
+		Log.Warning(FText::Format(LOCTEXT("FailedMetaInfo", "Failed to find 'payload_metainfo.txt' in the depot '{0}'\nThe source control backend will be unable to pull payloads, is your source control  config set up correctly?"),
 					FText::FromString(DepotRoot)));
 		
 		OnConnectionError();
@@ -299,7 +299,7 @@ bool FSourceControlBackend::Initialize(const FString& ConfigEntry)
 	{
 		FMessageLog Log("LogVirtualization");
 
-		Log.Warning(	FText::Format(LOCTEXT("FailedMetaInfo", "Failed to find 'payload_metainfo.txt' in the depot '{0}'\nThe source control backend will be unable to pull payloads, is your source control  config set up correctly?"),
+		Log.Warning(FText::Format(LOCTEXT("FailedMetaInfo", "Failed to find 'payload_metainfo.txt' in the depot '{0}'\nThe source control backend will be unable to pull payloads, is your source control  config set up correctly?"),
 					FText::FromString(DepotRoot)));
 
 		OnConnectionError();
@@ -722,6 +722,11 @@ bool FSourceControlBackend::TryApplySettingsFromConfigFiles(const FString& Confi
 		return false;
 	}
 
+	if (!DepotRoot.EndsWith(TEXT("/")))
+	{
+		DepotRoot.AppendChar(TEXT('/'));
+	}
+
 	// Now parse the optional config values
 
 	// Check to see if we should use partitioned clients or not. This is a perforce specific optimization to make the workspace churn cheaper on the server
