Make sure that all members in FVirtualizationManager set from the config file have appropriate default values set in case the config file fails to load for some reason.

#rnx

- The missing default values were taken from BaseEngine.ini.
- Made it a bit clearer in the header, which members are set from the config files and which are not.

diff --git a/Engine/Source/Developer/Virtualization/Private/VirtualizationManager.cpp b/Engine/Source/Developer/Virtualization/Private/VirtualizationManager.cpp
--- a/Engine/Source/Developer/Virtualization/Private/VirtualizationManager.cpp
+++ b/Engine/Source/Developer/Virtualization/Private/VirtualizationManager.cpp
@@ -264,6 +264,9 @@ FVirtualizationManager::FVirtualizationManager()
 	, bEnableCacheAfterPull(true)
 	, MinPayloadLength(0)
 	, BackendGraphName(TEXT("ContentVirtualizationBackendGraph_None"))
+	, FilteringMode(EPackageFilterMode::OptOut)
+	, bFilterEngineContent(true)
+	, bFilterEnginePluginContent(true)
 {
 }
 
diff --git a/Engine/Source/Developer/Virtualization/Private/VirtualizationManager.h b/Engine/Source/Developer/Virtualization/Private/VirtualizationManager.h
--- a/Engine/Source/Developer/Virtualization/Private/VirtualizationManager.h
+++ b/Engine/Source/Developer/Virtualization/Private/VirtualizationManager.h
@@ -209,6 +209,9 @@ private:
 	/** Determines if the default filtering behavior is to virtualize a payload or not */
 	bool ShouldVirtualizeAsDefault() const;
 	
+private:
+	// The following members are set from the config file
+
 	/** Are payloads allowed to be virtualized. Defaults to true. */
 	bool bEnablePayloadPushing;
 
@@ -222,7 +225,7 @@ private:
 	FString BackendGraphName;
 
 	/** The default filtering mode to apply if a payload is not matched with an option in UVirtualizationFilterSettings */
-	EPackageFilterMode FilteringMode = EPackageFilterMode::OptOut;
+	EPackageFilterMode FilteringMode;
 
 	/** Should payloads in engine content packages before filtered out and never virtualized */
 	bool bFilterEngineContent;
@@ -230,6 +233,8 @@ private:
 	/** Should payloads in engine plugin content packages before filtered out and never virtualized */
 	bool bFilterEnginePluginContent;
 	
+private:
+
 	/** The name of the current project */
 	FString ProjectName;
 
