Reworked the config file set up for the DDCBackend to remove the chance of Fatal errors.

#rnx

- The bucket entry was already marked as optional in the code documentation but was giving a fatal error if it was not set. The value is now optional and will default to 'BulkData' if not set.
- We no longer consider both the local and remote transfer policy being disabled as a fatal error, when the work UE-143676 is done this can be improved further.

diff --git a/Engine/Source/Developer/Virtualization/Private/VirtualizationDDCBackend.cpp b/Engine/Source/Developer/Virtualization/Private/VirtualizationDDCBackend.cpp
--- a/Engine/Source/Developer/Virtualization/Private/VirtualizationDDCBackend.cpp
+++ b/Engine/Source/Developer/Virtualization/Private/VirtualizationDDCBackend.cpp
@@ -18,10 +18,10 @@ static UE::DerivedData::FValueId ToDerivedDataValueId(const FIoHash& Id)
 }
 
 FDDCBackend::FDDCBackend(FStringView ProjectName, FStringView ConfigName, FStringView InDebugName)
-: IVirtualizationBackend(ConfigName, InDebugName, EOperations::Push | EOperations::Pull)
-, BucketName(TEXT("BulkData"))
-, TransferPolicy(UE::DerivedData::ECachePolicy::None)
-, QueryPolicy(UE::DerivedData::ECachePolicy::None)
+	: IVirtualizationBackend(ConfigName, InDebugName, EOperations::Push | EOperations::Pull)
+	, BucketName(TEXT("BulkData"))
+	, TransferPolicy(UE::DerivedData::ECachePolicy::None)
+	, QueryPolicy(UE::DerivedData::ECachePolicy::None)
 {
 	
 }
@@ -29,27 +29,26 @@ FDDCBackend::FDDCBackend(FStringView ProjectName, FStringView ConfigName, FStrin
 bool FDDCBackend::Initialize(const FString& ConfigEntry)
 {
 	TRACE_CPUPROFILER_EVENT_SCOPE(Initialize::Initialize);
-
-	if (!FParse::Value(*ConfigEntry, TEXT("Bucket="), BucketName))
+	
+	FString BucketNameIniFile;
+	if(FParse::Value(*ConfigEntry, TEXT("Bucket="), BucketNameIniFile))
 	{
-		UE_LOG(LogVirtualization, Fatal, TEXT("[%s] 'Bucket=' not found in the config file"), *GetDebugName());
+		BucketName = BucketNameIniFile;
 	}
 
 	bool bAllowLocal = true;
-	if (FParse::Bool(*ConfigEntry, TEXT("LocalStorage="), bAllowLocal))
-	{
-		UE_LOG(LogVirtualization, Display, TEXT("[%s] Use of local storage set to '%s"), *GetDebugName(), bAllowLocal ? TEXT("true") : TEXT("false"));
-	}
-
+	FParse::Bool(*ConfigEntry, TEXT("LocalStorage="), bAllowLocal);
+	
 	bool bAllowRemote = true;
-	if (FParse::Bool(*ConfigEntry, TEXT("RemoteStorage="), bAllowRemote))
-	{
-		UE_LOG(LogVirtualization, Display, TEXT("[%s] Use of remote storage set to '%s"), *GetDebugName(), bAllowRemote ? TEXT("true") : TEXT("false"));
-	}
+	FParse::Bool(*ConfigEntry, TEXT("RemoteStorage="), bAllowRemote);
+
+	UE_LOG(LogVirtualization, Log, TEXT("[%s] Bucket set to '%s"), *GetDebugName(), *BucketName);
+	UE_LOG(LogVirtualization, Log, TEXT("[%s] Use of local storage set to '%s"), *GetDebugName(), bAllowLocal ? TEXT("true") : TEXT("false"));
+	UE_LOG(LogVirtualization, Log, TEXT("[%s] Use of remote storage set to '%s"), *GetDebugName(), bAllowRemote ? TEXT("true") : TEXT("false"));
 
 	if (!bAllowLocal && !bAllowRemote)
 	{
-		UE_LOG(LogVirtualization, Fatal, TEXT("[%s] LocalStorage and RemoteStorage cannot both be disabled"), *GetDebugName());
+		UE_LOG(LogVirtualization, Error, TEXT("[%s] LocalStorage and RemoteStorage cannot both be disabled"), *GetDebugName());
 		return false;
 	}
 
