Add a config file option to allow projects to customize (or disable) the tag applied to changelist descriptions once the asset virtualization process has run.

#rnx

- The option is [Core.VirtualizationModule]VirtualizationProcessTag=
- If left blank then no tag will be added
- The internal virtualization process code no longer needs to return a tag since it will now be applied by the virtualization manager instead.
-- Note that we still follow the same logic, we only stamp if no errors were encountered.

diff --git a/Engine/Config/BaseEngine.ini b/Engine/Config/BaseEngine.ini
--- a/Engine/Config/BaseEngine.ini
+++ b/Engine/Config/BaseEngine.ini
@@ -1241,6 +1241,7 @@ EnablePushToBackend=False
 EnableCacheAfterPull=True
 MinPayloadLength=0
 BackendGraph=ContentVirtualizationBackendGraph_None
+VirtualizationProcessTag=#virtualized
 AllowSubmitIfVirtualizationFailed=false
 FilterMode=OptOut
 FilterEngineContent=True
diff --git a/Engine/Source/Developer/Virtualization/Private/PackageVirtualizationProcess.cpp b/Engine/Source/Developer/Virtualization/Private/PackageVirtualizationProcess.cpp
--- a/Engine/Source/Developer/Virtualization/Private/PackageVirtualizationProcess.cpp
+++ b/Engine/Source/Developer/Virtualization/Private/PackageVirtualizationProcess.cpp
@@ -197,7 +197,7 @@ private:
 };
 #endif //ENABLE_FILTERING_HACK
 
-void VirtualizePackages(const TArray<FString>& FilesToSubmit, TArray<FText>& OutDescriptionTags, TArray<FText>& OutErrors)
+void VirtualizePackages(const TArray<FString>& FilesToSubmit, TArray<FText>& OutErrors)
 {
 	TRACE_CPUPROFILER_EVENT_SCOPE(UE::Virtualization::VirtualizePackages);
 
@@ -535,15 +535,6 @@ void VirtualizePackages(const TArray<FString>& FilesToSubmit, TArray<FText>& Out
 		}
 	}
 
-	// If we had no new errors add the validation tag to indicate that the packages are safe for submission. 
-	// TODO: Currently this is a simple tag to make it easier for us to track which assets were submitted via the
-	// virtualization process in a test project. This should be expanded when we add proper p4 server triggers.
-	if (NumErrors == OutErrors.Num())
-	{
-		FText Tag = FText::FromString(TEXT("#virtualized"));
-		OutDescriptionTags.Add(Tag);
-	}
-
 	const double TimeInSeconds = FPlatformTime::Seconds() - StartTime;
 	UE_LOG(LogVirtualization, Verbose, TEXT("Virtualization pre submit check took %.3f(s)"), TimeInSeconds);
 }
diff --git a/Engine/Source/Developer/Virtualization/Private/PackageVirtualizationProcess.h b/Engine/Source/Developer/Virtualization/Private/PackageVirtualizationProcess.h
--- a/Engine/Source/Developer/Virtualization/Private/PackageVirtualizationProcess.h
+++ b/Engine/Source/Developer/Virtualization/Private/PackageVirtualizationProcess.h
@@ -9,7 +9,6 @@
 namespace UE::Virtualization
 {
 
-/** @See IVirtualizationSystem::TryVirtualizePackages */
-void VirtualizePackages(const TArray<FString>& FilesToSubmit, TArray<FText>& OutDescriptionTags, TArray<FText>& OutErrors);
+void VirtualizePackages(const TArray<FString>& FilesToSubmit, TArray<FText>& OutErrors);
 
 } // namespace UE::Virtualization
diff --git a/Engine/Source/Developer/Virtualization/Private/VirtualizationManager.cpp b/Engine/Source/Developer/Virtualization/Private/VirtualizationManager.cpp
--- a/Engine/Source/Developer/Virtualization/Private/VirtualizationManager.cpp
+++ b/Engine/Source/Developer/Virtualization/Private/VirtualizationManager.cpp
@@ -264,6 +264,7 @@ FVirtualizationManager::FVirtualizationManager()
 	, bEnableCacheAfterPull(true)
 	, MinPayloadLength(0)
 	, BackendGraphName(TEXT("ContentVirtualizationBackendGraph_None"))
+	, VirtualizationProcessTag(TEXT("#virtualized"))
 	, FilteringMode(EPackageFilterMode::OptOut)
 	, bFilterEngineContent(true)
 	, bFilterEnginePluginContent(true)
@@ -641,7 +642,16 @@ EVirtualizationResult FVirtualizationManager::TryVirtualizePackages(const TArray
 	OutDescriptionTags.Reset();
 	OutErrors.Reset();
 
-	UE::Virtualization::VirtualizePackages(FilesToVirtualize, OutDescriptionTags, OutErrors);
+	UE::Virtualization::VirtualizePackages(FilesToVirtualize, OutErrors);
+
+	// If we had no new errors add the validation tag to indicate that the packages are safe for submission. 
+	// TODO: Currently this is a simple tag to make it easier for us to track which assets were submitted via the
+	// virtualization process in a test project. This should be expanded when we add proper p4 server triggers.
+	if (OutErrors.IsEmpty() && !VirtualizationProcessTag.IsEmpty())
+	{
+		FText Tag = FText::FromString(VirtualizationProcessTag);
+		OutDescriptionTags.Add(Tag);
+	}
 
 	return OutErrors.IsEmpty() ? EVirtualizationResult::Success : EVirtualizationResult::Failed;
 }
@@ -781,6 +791,17 @@ void FVirtualizationManager::ApplySettingsFromConfigFiles(const FConfigFile& Con
 		UE_LOG(LogVirtualization, Error, TEXT("Failed to load [Core.VirtualizationModule].BackendGraph from config file!"));
 	}
 
+	FString VirtualizationProcessTagFromIni;
+	if (ConfigFile.GetString(TEXT("Core.VirtualizationModule"), TEXT("VirtualizationProcessTag"), VirtualizationProcessTagFromIni))
+	{
+		VirtualizationProcessTag = VirtualizationProcessTagFromIni;
+		UE_LOG(LogVirtualization, Display, TEXT("\tVirtualizationProcessTag : %s"), *VirtualizationProcessTag);
+	}
+	else
+	{
+		UE_LOG(LogVirtualization, Error, TEXT("Failed to load [Core.VirtualizationModule].VirtualizationProcessTag from config file!"));
+	}
+
 	FString FilterModeFromIni;
 	if (ConfigFile.GetString(TEXT("Core.ContentVirtualization"), TEXT("FilterMode"), FilterModeFromIni) ||
 		ConfigFile.GetString(TEXT("Core.VirtualizationModule"), TEXT("FilterMode"), FilterModeFromIni))
diff --git a/Engine/Source/Developer/Virtualization/Private/VirtualizationManager.h b/Engine/Source/Developer/Virtualization/Private/VirtualizationManager.h
--- a/Engine/Source/Developer/Virtualization/Private/VirtualizationManager.h
+++ b/Engine/Source/Developer/Virtualization/Private/VirtualizationManager.h
@@ -226,6 +226,9 @@ private:
 	/** The name of the backend graph to load from the config ini file that will describe the backend hierarchy */
 	FString BackendGraphName;
 
+	/** The tag that will be returned when the virtualization process has run, commonly used to post fix changelist descriptions */
+	FString VirtualizationProcessTag;
+
 	/** The default filtering mode to apply if a payload is not matched with an option in UVirtualizationFilterSettings */
 	EPackageFilterMode FilteringMode;
 
