Add a config file option to allow payloads stored in .umap and associated "_builddata.uasset" files to be filtered from the virtualization process. This option defaults to true.

#rnx

### Problem
- The payloads stored in map files tend to change more than other assets and would cause a lot more churn in the VA system.
- Some other systems like the landscape component are not able to sensibly continue if their payloads cannot be accessed (heightmaps for exmaple) and would prefer not to allow virtualization.
- As a short term fix we need an option to disable the virtualization on all payloads in map files. Future improvements to the filtering system will allow systems to more easily opt their payloads out of virtualization. When this is functional we might want to change the default from true to false.

### Feature
- The config optionf for this filtering is "[Core.VirtualizationModule]FilterMapContent=True"
- Testing if the owning UObject for a payload is in a umap can get tricky, because we not only need to check the umap but we also need to check if it is in a "_builddata.uasset" file, which is an additional file we store next to a umap containing things like lightmaps etc.
- At the moment we check for this by finding the outermost object for the given owner and check to see if it is a ULevel, UWorld or UMapBuildDataRegistry. This is a bit of a kludge but the types we need to check against are not accessible by this module and making them accessible will pull in a lot of dependencies that we'd prefer not to add.
-- One improvement might be to tag the FLinkerSave with the info we need and passing that into the serialization process rather than trying to work it out ourselves but I am wary of making that change until we are 100% sure that we want to keep this feature.

### Refactor
- Removed IVirtualizationSystem::IsDisabledForObject and replaced it with ::FilterPayload which can return multiple reasons for preventing a payload from virtualizing. (the method was added during 5.1 development so it should be fine to just replace it without deprecation)
-- The original behaviour for FVirtualizationManager::IsDisabledForObject has been moved to FVirtualizationManager::ShouldVirtualizeAsset
- Added a new header to declare enums/types used by the various parts of the virtualization system and started by moving EPayloadFilterReason there from the package trailer header. This allows both the core API and PackageTrailer to use EPayloadFilterReason without creating overburdened header dependencies.
-- EPayloadFilterReason has moved from the UE namespace to UE::Virtualization so the package trailer code needed updating accordingly.
- EditorBulkData will ask the virtualization system for the base filter reason, then add it's own reasons if UE_ENABLE_VIRTUALIZATION_TOGGLE is enabled. This bit of code will be removed for 5.1

diff --git a/Engine/Config/BaseEngine.ini b/Engine/Config/BaseEngine.ini
--- a/Engine/Config/BaseEngine.ini
+++ b/Engine/Config/BaseEngine.ini
@@ -1246,6 +1246,7 @@ AllowSubmitIfVirtualizationFailed=false
 FilterMode=OptOut
 FilterEngineContent=True
 FilterEnginePluginContent=True
+FilterMapContent=True
 ; The following assets should not virtualize until it is confirmed that they do not crash when virtualized payloads are missing
 +DisabledAsset=StaticMesh
 +DisabledAsset=SoundWave
diff --git a/Engine/Source/Developer/Virtualization/Private/VirtualizationManager.cpp b/Engine/Source/Developer/Virtualization/Private/VirtualizationManager.cpp
--- a/Engine/Source/Developer/Virtualization/Private/VirtualizationManager.cpp
+++ b/Engine/Source/Developer/Virtualization/Private/VirtualizationManager.cpp
@@ -268,6 +268,7 @@ FVirtualizationManager::FVirtualizationManager()
 	, FilteringMode(EPackageFilterMode::OptOut)
 	, bFilterEngineContent(true)
 	, bFilterEnginePluginContent(true)
+	, bFilterMapContent(true)
 	, bAllowSubmitIfVirtualizationFailed(false)
 {
 }
@@ -376,22 +377,34 @@ bool FVirtualizationManager::IsPushingEnabled(EStorageType StorageType) const
 	}
 }
 
-bool FVirtualizationManager::IsDisabledForObject(const UObject* OwnerObject) const
+EPayloadFilterReason FVirtualizationManager::FilterPayload(const UObject* Owner) const
 {
-	if (OwnerObject == nullptr)
+	UE::Virtualization::EPayloadFilterReason PayloadFilter = UE::Virtualization::EPayloadFilterReason::None;
+	if (!ShouldVirtualizeAsset(Owner))
 	{
-		return false;
+		PayloadFilter |= UE::Virtualization::EPayloadFilterReason::Asset;
 	}
 
-	const UClass* OwnerClass = OwnerObject->GetClass();
-	if (OwnerClass == nullptr)
+	// TODO: If we keep this feature long term then we might want to work this out in SavePackage.cpp and pass the info
+	// via FLinkerSave rather than the following code.
+	if (bFilterMapContent)
 	{
-		// TODO: Not actually sure if the class being nullptr is reasonable or if we should warn/error here?
-		return false;
+		if (const UObject* Outer = Owner->GetOutermostObject())
+		{
+			if (const UClass* OuterClass = Outer->GetClass())
+			{
+				const FName OuterClassName = OuterClass->GetFName();
+				if (OuterClassName == FName("Level") ||
+					OuterClassName == FName("World") ||
+					OuterClassName == FName("MapBuildDataRegistry"))
+				{
+					PayloadFilter |= UE::Virtualization::EPayloadFilterReason::MapContent;
+				}
+			}
+		}
 	}
 
-	FName ClassName = OwnerClass->GetFName();
-	return DisabledAssetTypes.Find(ClassName) != nullptr;
+	return PayloadFilter;
 }
 
 bool FVirtualizationManager::AllowSubmitIfVirtualizationFailed() const
@@ -844,6 +857,17 @@ void FVirtualizationManager::ApplySettingsFromConfigFiles(const FConfigFile& Con
 		UE_LOG(LogVirtualization, Error, TEXT("Failed to load [Core.VirtualizationModule].FilterEnginePluginContent from config file!"));
 	}
 
+	// Optional
+	bool bFilterMapContentFromIni = false;
+	if (ConfigFile.GetBool(TEXT("Core.VirtualizationModule"), TEXT("FilterMapContent"), bFilterMapContentFromIni))
+	{
+		bFilterMapContent = bFilterMapContentFromIni;
+		UE_LOG(LogVirtualization, Display, TEXT("\tFilterMapContent : %s"), bFilterMapContent ? TEXT("true") : TEXT("false"));
+	}
+	else
+	{
+		UE_LOG(LogVirtualization, Error, TEXT("Failed to load [Core.VirtualizationModule].FilterMapContent from config file!"));
+	}
 
 	// Optional
 	TArray<FString> DisabledAssetTypesFromIni;
@@ -1416,6 +1440,25 @@ FCompressedBuffer FVirtualizationManager::PullDataFromBackend(IVirtualizationBac
 	return Payload;
 }
 
+
+bool FVirtualizationManager::ShouldVirtualizeAsset(const UObject* OwnerObject) const
+{
+	if (OwnerObject == nullptr)
+	{
+		return true;
+	}
+
+	const UClass* OwnerClass = OwnerObject->GetClass();
+	if (OwnerClass == nullptr)
+	{
+		// TODO: Not actually sure if the class being nullptr is reasonable or if we should warn/error here?
+		return true;
+	}
+
+	const FName ClassName = OwnerClass->GetFName();
+	return DisabledAssetTypes.Find(ClassName) == nullptr;
+}
+
 bool FVirtualizationManager::ShouldVirtualizePackage(const FPackagePath& PackagePath) const
 {
 	TRACE_CPUPROFILER_EVENT_SCOPE(FVirtualizationManager::ShouldVirtualizePackage);
diff --git a/Engine/Source/Developer/Virtualization/Private/VirtualizationManager.h b/Engine/Source/Developer/Virtualization/Private/VirtualizationManager.h
--- a/Engine/Source/Developer/Virtualization/Private/VirtualizationManager.h
+++ b/Engine/Source/Developer/Virtualization/Private/VirtualizationManager.h
@@ -72,6 +72,7 @@ class FOutputDevice;
  * FilterMode=OptIn/OptOut					When 'OptIn' payloads will be virtualized by default, when 'OptOut' they will not be virtualized by default
  * FilterEngineContent=True/False			When true any payload from a package under Engine/Content/.. will be excluded from virtualization
  * FilterEnginePluginContent=True/False		When true any payload from a package under Engine/Plugins/../Content/.. will be excluded from virtualization
+ * FilterMapContent=True/False				When true any payload stored in a .umap or _BuildData.uasset file will be excluded from virtualization
  * 
  * PackagePath Setup:
  * 
@@ -129,7 +130,7 @@ private:
 	virtual bool IsEnabled() const override;
 	virtual bool IsPushingEnabled(EStorageType StorageType) const override;
 
-	virtual bool IsDisabledForObject(const UObject* Owner) const override;
+	virtual EPayloadFilterReason FilterPayload(const UObject* Owner) const override;
 
 	virtual bool AllowSubmitIfVirtualizationFailed() const override;
 	
@@ -184,6 +185,8 @@ private:
 	FCompressedBuffer PullDataFromAllBackends(const FIoHash& Id);
 	FCompressedBuffer PullDataFromBackend(IVirtualizationBackend& Backend, const FIoHash& Id);
 
+	bool ShouldVirtualizeAsset(const UObject* Owner) const;
+
 	/** 
 	 * Determines if a package path should be virtualized or not based on any exclusion/inclusion patterns
 	 * that might have been set in UVirtualizationFilterSettings.
@@ -232,12 +235,15 @@ private:
 	/** The default filtering mode to apply if a payload is not matched with an option in UVirtualizationFilterSettings */
 	EPackageFilterMode FilteringMode;
 
-	/** Should payloads in engine content packages before filtered out and never virtualized */
+	/** Should payloads in engine content packages be filtered out and never virtualized */
 	bool bFilterEngineContent;
 	
-	/** Should payloads in engine plugin content packages before filtered out and never virtualized */
+	/** Should payloads in engine plugin content packages be filtered out and never virtualized */
 	bool bFilterEnginePluginContent;
 
+	/** Should payloads in .umap files (or associated _BuildData files) be filtered out and never virtualized */
+	bool bFilterMapContent;
+
 	/** Should file submits be allowed to continue if a call to TryVirtualizePackages fails */
 	bool bAllowSubmitIfVirtualizationFailed;
 	
diff --git a/Engine/Source/Runtime/Core/Private/Virtualization/VirtualizationSystem.cpp b/Engine/Source/Runtime/Core/Private/Virtualization/VirtualizationSystem.cpp
--- a/Engine/Source/Runtime/Core/Private/Virtualization/VirtualizationSystem.cpp
+++ b/Engine/Source/Runtime/Core/Private/Virtualization/VirtualizationSystem.cpp
@@ -42,9 +42,9 @@ public:
 		return false;
 	}
 
-	virtual bool IsDisabledForObject(const UObject* Owner) const override
+	virtual EPayloadFilterReason FilterPayload(const UObject* Owner) const override
 	{
-		return false;
+		return EPayloadFilterReason::None;
 	}
 
 	virtual bool AllowSubmitIfVirtualizationFailed() const override
diff --git a/Engine/Source/Runtime/Core/Public/Virtualization/VirtualizationSystem.h b/Engine/Source/Runtime/Core/Public/Virtualization/VirtualizationSystem.h
--- a/Engine/Source/Runtime/Core/Public/Virtualization/VirtualizationSystem.h
+++ b/Engine/Source/Runtime/Core/Public/Virtualization/VirtualizationSystem.h
@@ -18,6 +18,7 @@
 #include "Templates/UniquePtr.h"
 #include "Templates/UnrealTemplate.h"
 #include "UObject/NameTypes.h"
+#include "Virtualization/VirtualizationTypes.h"
 
 class FPackagePath;
 class FText;
@@ -351,14 +352,14 @@ public:
 	virtual bool IsPushingEnabled(EStorageType StorageType) const = 0;
 
 	/** 
-	 * Poll to see if virtualization is disabled for the given asset type.
+	 * Run checks to see if the payload can be virtualized or not.
 	 * 
-	 * @param	Owner	The object to be tested, assumed to be an asset that 
-	 *					can own virtualized payloads.
+	 * @param	Owner	The UObject that contains the payload, this can be nullptr if the payload is not owned by one
 	 * 
-	 * @return	True if payloads owned by this object should never virtualized
+	 * @return	The reasons why the payload can not be virtualized, encoded as a bitfield of EPayloadFilterReason. 
+	 *			This bitfield will be EPayloadFilterReason::None (0) if no reason was found.
 	 */
-	virtual bool IsDisabledForObject(const UObject* Owner) const = 0;
+	virtual EPayloadFilterReason FilterPayload(const UObject* Owner) const = 0;
 
 	/** 
 	 * Poll to see if the virtualization process failing when submitting a collection of files to source
diff --git a/Engine/Source/Runtime/Core/Public/Virtualization/VirtualizationTypes.h b/Engine/Source/Runtime/Core/Public/Virtualization/VirtualizationTypes.h
new file mode 100644
--- /dev/null
+++ b/Engine/Source/Runtime/Core/Public/Virtualization/VirtualizationTypes.h
@@ -0,0 +1,29 @@
+// Copyright Epic Games, Inc. All Rights Reserved.
+
+#pragma once
+
+#include "Misc/EnumClassFlags.h"
+
+namespace UE::Virtualization
+{
+
+/** This enum describes the reasons why a payload may not be virtualized */
+enum class EPayloadFilterReason : uint16
+{
+	/** Not filtered, the payload can be virtualized */
+	None				= 0,
+	/** Filtered due to the asset type of the owning UObject */
+	Asset				= 1 << 0,
+	/** Filtered due to the path of the owning UPackage */
+	Path				= 1 << 1,
+	/** Filtered because the payload size is below the minimum size for virtualization */
+	MinSize				= 1 << 2,
+	/** Filtered because the owning editor bulkdata had virtualization disabled programmatically */
+	EditorBulkDataCode	= 1 << 3,
+	/** Filtered because the package is either a UMap or the owning editor bulkdata is under a UMapBuildDataRegistry */
+	MapContent			= 1 << 4,
+};
+
+ENUM_CLASS_FLAGS(EPayloadFilterReason);
+
+} //UE::Virtualization
\ No newline at end of file
diff --git a/Engine/Source/Runtime/CoreUObject/Private/Serialization/EditorBulkData.cpp b/Engine/Source/Runtime/CoreUObject/Private/Serialization/EditorBulkData.cpp
--- a/Engine/Source/Runtime/CoreUObject/Private/Serialization/EditorBulkData.cpp
+++ b/Engine/Source/Runtime/CoreUObject/Private/Serialization/EditorBulkData.cpp
@@ -20,6 +20,7 @@
 #include "UObject/UObjectGlobals.h"
 #include "UObject/SavePackage.h"
 #include "Virtualization/VirtualizationSystem.h"
+#include "Virtualization/VirtualizationTypes.h"
 
 //#if WITH_EDITORONLY_DATA
 
@@ -1702,11 +1703,14 @@ void FEditorBulkData::SerializeToPackageTrailer(FLinkerSave& LinkerSave, FCompre
 		}
 	};
 
-	EPayloadFilterReason PayloadFilter = EPayloadFilterReason::None;
-	if (bSkipVirtualization || UE::Virtualization::IVirtualizationSystem::Get().IsDisabledForObject(Owner))
+	UE::Virtualization::EPayloadFilterReason PayloadFilter = UE::Virtualization::IVirtualizationSystem::Get().FilterPayload(Owner);
+
+#if UE_ENABLE_VIRTUALIZATION_TOGGLE
+	if (bSkipVirtualization)
 	{
-		PayloadFilter = EPayloadFilterReason::Asset;
+		PayloadFilter |= UE::Virtualization::EPayloadFilterReason::Asset;
 	}
+#endif //UE_ENABLE_VIRTUALIZATION_TOGGLE
 
 	LinkerSave.PackageTrailerBuilder->AddPayload(PayloadContentId, MoveTemp(PayloadToSerialize), PayloadFilter, MoveTemp(OnPayloadWritten));
 }
diff --git a/Engine/Source/Runtime/CoreUObject/Private/UObject/PackageTrailer.cpp b/Engine/Source/Runtime/CoreUObject/Private/UObject/PackageTrailer.cpp
--- a/Engine/Source/Runtime/CoreUObject/Private/UObject/PackageTrailer.cpp
+++ b/Engine/Source/Runtime/CoreUObject/Private/UObject/PackageTrailer.cpp
@@ -205,7 +205,7 @@ FPackageTrailerBuilder::FPackageTrailerBuilder(FString&& InDebugContext)
 {
 }
 
-void FPackageTrailerBuilder::AddPayload(const FIoHash& Identifier, FCompressedBuffer Payload, EPayloadFilterReason FilterFlags, AdditionalDataCallback&& Callback)
+void FPackageTrailerBuilder::AddPayload(const FIoHash& Identifier, FCompressedBuffer Payload, UE::Virtualization::EPayloadFilterReason FilterFlags, AdditionalDataCallback&& Callback)
 {
 	Callbacks.Emplace(MoveTemp(Callback));
 
@@ -242,7 +242,7 @@ bool FPackageTrailerBuilder::UpdatePayloadAsLocal(const FIoHash& Identifier, FCo
 	if (!Identifier.IsZero() && VirtualizedEntries.Remove(Identifier) > 0)
 	{
 		check(LocalEntries.Find(Identifier) == nullptr);
-		LocalEntries.Add(Identifier, LocalEntry(MoveTemp(Payload), EPayloadFilterReason::None));
+		LocalEntries.Add(Identifier, LocalEntry(MoveTemp(Payload), UE::Virtualization::EPayloadFilterReason::None));
 
 		return true;
 	}
@@ -780,7 +780,7 @@ TArray<FIoHash> FPackageTrailer::GetPayloads(EPayloadFilter Filter) const
 		switch (Filter)
 		{
 			case EPayloadFilter::CanVirtualize:
-				if (Entry.IsLocal() && Entry.FilterFlags == EPayloadFilterReason::None)
+				if (Entry.IsLocal() && Entry.FilterFlags == UE::Virtualization::EPayloadFilterReason::None)
 				{
 					Identifiers.Add(Entry.Identifier);
 				}
@@ -803,7 +803,7 @@ int32 FPackageTrailer::GetNumPayloads(EPayloadFilter Filter) const
 		case EPayloadFilter::CanVirtualize:
 			Count = (int32)Algo::CountIf(Header.PayloadLookupTable, [](const Private::FLookupTableEntry& Entry)
 				{
-					return Entry.AccessMode == EPayloadAccessMode::Local && Entry.FilterFlags == EPayloadFilterReason::None;
+					return Entry.AccessMode == EPayloadAccessMode::Local && Entry.FilterFlags == UE::Virtualization::EPayloadFilterReason::None;
 				});
 			break;
 
@@ -890,9 +890,9 @@ bool FindPayloadsInPackageFile(const FPackagePath& PackagePath, EPayloadStorageT
 
 } //namespace UE
 
-FString LexToString(UE::EPayloadFilterReason FilterFlags)
+FString LexToString(UE::Virtualization::EPayloadFilterReason FilterFlags)
 {
-	using namespace UE;
+	using namespace UE::Virtualization;
 
 	if (FilterFlags == EPayloadFilterReason::None)
 	{
@@ -932,13 +932,22 @@ FString LexToString(UE::EPayloadFilterReason FilterFlags)
 			Builder << TEXT("EditorBulkDataCode");
 		}
 
+		if (EnumHasAllFlags(FilterFlags, EPayloadFilterReason::MapContent))
+		{
+			AddSeparatorIfNeeded();
+			Builder << TEXT("MapContent");
+		}
+
 		// In case a new entry was added to EPayloadFilterReason without this function being
 		// updated we need to check if any of the other bits are set and if they are print 
 		// and unknown entry.
 		// We don't want a Max or Count entry being added to EPayloadFilterReason as we'd then
 		// need to validate all places taking the enum to make sure that nobody is using it.
-		const uint16 Mask = ~uint16(EPayloadFilterReason::Asset | EPayloadFilterReason::Path |
-			EPayloadFilterReason::MinSize | EPayloadFilterReason::EditorBulkDataCode);
+		const uint16 Mask = ~uint16(	EPayloadFilterReason::Asset |
+										EPayloadFilterReason::Path |
+										EPayloadFilterReason::MinSize |
+										EPayloadFilterReason::EditorBulkDataCode |
+										EPayloadFilterReason::MapContent);
 
 		const bool bHasUnknownBits = ((uint16)FilterFlags & Mask) != 0;
 		if (bHasUnknownBits)
diff --git a/Engine/Source/Runtime/CoreUObject/Public/UObject/PackageTrailer.h b/Engine/Source/Runtime/CoreUObject/Public/UObject/PackageTrailer.h
--- a/Engine/Source/Runtime/CoreUObject/Public/UObject/PackageTrailer.h
+++ b/Engine/Source/Runtime/CoreUObject/Public/UObject/PackageTrailer.h
@@ -12,6 +12,7 @@
 #include "Templates/Function.h"
 #include "Templates/UniquePtr.h"
 #include "UObject/NameTypes.h"
+#include "Virtualization/VirtualizationTypes.h"
 
 class FArchive;
 class FLinkerSave;
@@ -123,22 +124,6 @@ enum class EPayloadFlags : uint16
 
 ENUM_CLASS_FLAGS(EPayloadFlags);
 
-enum class EPayloadFilterReason : uint16
-{
-	/** Not filtered */
-	None				= 0,
-	/** Filtered due to the asset type of the owning UObject */
-	Asset				= 1 << 0,
-	/** Filtered due to the path of the owning UPackage */
-	Path				= 1 << 1,
-	/** Filtered because the payload size is below the minimum size for virtualization */
-	MinSize				= 1 << 2,
-	/** Filtered because the owning editor bulkdata had virtualization disabled programmatically */
-	EditorBulkDataCode	= 1 << 3
-};
-
-ENUM_CLASS_FLAGS(EPayloadFilterReason);
-
 enum class EPackageTrailerVersion : uint32;
 
 namespace Private
@@ -185,7 +170,7 @@ struct FLookupTableEntry
 	/** Bitfield of flags, see @UE::EPayloadFlags */
 	EPayloadFlags Flags = EPayloadFlags::None;
 	/** Bitfield of flags showing if the payload allowed to be virtualized or the reason why it cannot be virtualized, see @UE::EPayloadFilterReason */
-	EPayloadFilterReason FilterFlags = EPayloadFilterReason::None;
+	Virtualization::EPayloadFilterReason FilterFlags = Virtualization::EPayloadFilterReason::None;
 
 	EPayloadAccessMode AccessMode = EPayloadAccessMode::Local;
 };
@@ -202,7 +187,7 @@ struct FPayloadInfo
 
 	EPayloadAccessMode AccessMode = EPayloadAccessMode::Local;
 	EPayloadFlags Flags = EPayloadFlags::None;
-	EPayloadFilterReason FilterFlags = EPayloadFilterReason::None;
+	Virtualization::EPayloadFilterReason FilterFlags = Virtualization::EPayloadFilterReason::None;
 };
 
 /** 
@@ -266,7 +251,7 @@ public:
 	 * @param Flags			The custom flags to be applied to the payload
 	 * @param Callback		This callback will be invoked once the FPackageTrailer has been built and appended to disk.
 	 */
-	void AddPayload(const FIoHash& Identifier, FCompressedBuffer Payload, EPayloadFilterReason Filter, AdditionalDataCallback&& Callback);
+	void AddPayload(const FIoHash& Identifier, FCompressedBuffer Payload, UE::Virtualization::EPayloadFilterReason Filter, AdditionalDataCallback&& Callback);
 
 	/**
 	 * Adds an already virtualized payload to the builder to be written to the trailer. When the trailer is written
@@ -326,7 +311,7 @@ private:
 	struct LocalEntry
 	{
 		LocalEntry() = default;
-		LocalEntry(FCompressedBuffer&& InPayload, EPayloadFilterReason InFilterFlags)
+		LocalEntry(FCompressedBuffer&& InPayload, Virtualization::EPayloadFilterReason InFilterFlags)
 			: Payload(InPayload)
 			, FilterFlags(InFilterFlags)
 		{
@@ -335,7 +320,7 @@ private:
 		~LocalEntry() = default;
 
 		FCompressedBuffer Payload;
-		EPayloadFilterReason FilterFlags = EPayloadFilterReason::None;
+		Virtualization::EPayloadFilterReason FilterFlags = Virtualization::EPayloadFilterReason::None;
 	};
 
 	/** All of the data required to add a reference to a payload stored in another trailer */
@@ -567,4 +552,4 @@ private:
 
 } //namespace UE
 
-[[nodiscard]] COREUOBJECT_API FString LexToString(UE::EPayloadFilterReason FilterFlags);
+[[nodiscard]] COREUOBJECT_API FString LexToString(UE::Virtualization::EPayloadFilterReason FilterFlags);
