The DumpPackagePayloadInfo console command now accepts fullobject paths as args as well as filepaths.

#rnx

diff --git a/Engine/Source/Developer/Virtualization/Private/DumpPackagePayloadInfoCommand.cpp b/Engine/Source/Developer/Virtualization/Private/DumpPackagePayloadInfoCommand.cpp
--- a/Engine/Source/Developer/Virtualization/Private/DumpPackagePayloadInfoCommand.cpp
+++ b/Engine/Source/Developer/Virtualization/Private/DumpPackagePayloadInfoCommand.cpp
@@ -4,6 +4,7 @@
 #include "Logging/LogMacros.h"
 #include "Misc/PackagePath.h"
 #include "UObject/PackageTrailer.h"
+#include "Misc/PackageName.h"
 
 namespace UE
 {
@@ -46,8 +47,18 @@ void DumpPackagePayloadInfo(const TArray<FString>& Args)
 	for (const FString& Arg : Args)
 	{
 		FPackagePath Path;
+		FStringView ObjectPathString;
 
-		if (FPackagePath::TryFromMountedName(Arg, Path))
+		if (FPackageName::ParseExportTextPath(Arg, nullptr /*OutClassName*/, &ObjectPathString))
+		{
+			ObjectPathString = FPackageName::ObjectPathToPackageName(ObjectPathString);
+		}
+		else
+		{
+			ObjectPathString = Arg;
+		}
+
+		if (FPackagePath::TryFromMountedName(ObjectPathString, Path))
 		{
 			FPackageTrailer Trailer;
 			if (!FPackageTrailer::TryLoadFromPackage(Path, Trailer))
