The virtualization manager will now read the newer config file location (Core.ContentVirtualization) first and only fall back to the older location (Core.ContentVirtualization) if it was not found.

#rnx

- This will help if someone has updated their set up to the new location but not removed the old entries.
- We still warn that the old enteries exist, even if not used.

diff --git a/Engine/Source/Developer/Virtualization/Private/VirtualizationManager.cpp b/Engine/Source/Developer/Virtualization/Private/VirtualizationManager.cpp
--- a/Engine/Source/Developer/Virtualization/Private/VirtualizationManager.cpp
+++ b/Engine/Source/Developer/Virtualization/Private/VirtualizationManager.cpp
@@ -757,8 +757,8 @@ void FVirtualizationManager::ApplySettingsFromConfigFiles(const FConfigFile& Con
 	// to it's own specific "Core.VirtualizationModule" section. This duplication can be removed before we ship 5.1
 
 	bool bEnablePayloadPushingFromIni = false;
-	if (ConfigFile.GetBool(TEXT("Core.ContentVirtualization"), TEXT("EnablePushToBackend"), bEnablePayloadPushingFromIni) ||
-		ConfigFile.GetBool(TEXT("Core.VirtualizationModule"), TEXT("EnablePushToBackend"), bEnablePayloadPushingFromIni))
+	if (ConfigFile.GetBool(TEXT("Core.VirtualizationModule"), TEXT("EnablePushToBackend"), bEnablePayloadPushingFromIni) ||
+		ConfigFile.GetBool(TEXT("Core.ContentVirtualization"), TEXT("EnablePushToBackend"), bEnablePayloadPushingFromIni))
 	{
 		bEnablePayloadPushing = bEnablePayloadPushingFromIni;
 		UE_LOG(LogVirtualization, Display, TEXT("\tEnablePushToBackend : %s"), bEnablePayloadPushing ? TEXT("true") : TEXT("false") );
@@ -769,8 +769,8 @@ void FVirtualizationManager::ApplySettingsFromConfigFiles(const FConfigFile& Con
 	}
 
 	bool bEnableCacheAfterPullFromIni = false;
-	if (ConfigFile.GetBool(TEXT("Core.ContentVirtualization"), TEXT("EnableCacheAfterPull"), bEnableCacheAfterPullFromIni) ||
-		ConfigFile.GetBool(TEXT("Core.VirtualizationModule"), TEXT("EnableCacheAfterPull"), bEnableCacheAfterPullFromIni))
+	if (ConfigFile.GetBool(TEXT("Core.VirtualizationModule"), TEXT("EnableCacheAfterPull"), bEnableCacheAfterPullFromIni) ||
+		ConfigFile.GetBool(TEXT("Core.ContentVirtualization"), TEXT("EnableCacheAfterPull"), bEnableCacheAfterPullFromIni))
 	{
 		bEnableCacheAfterPull = bEnableCacheAfterPullFromIni;
 		UE_LOG(LogVirtualization, Display, TEXT("\tCachePulledPayloads : %s"), bEnableCacheAfterPull ? TEXT("true") : TEXT("false"));
@@ -781,8 +781,8 @@ void FVirtualizationManager::ApplySettingsFromConfigFiles(const FConfigFile& Con
 	}
 
 	int64 MinPayloadLengthFromIni = 0;
-	if (ConfigFile.GetInt64(TEXT("Core.ContentVirtualization"), TEXT("MinPayloadLength"), MinPayloadLengthFromIni) ||
-		ConfigFile.GetInt64(TEXT("Core.VirtualizationModule"), TEXT("MinPayloadLength"), MinPayloadLengthFromIni))
+	if (ConfigFile.GetInt64(TEXT("Core.VirtualizationModule"), TEXT("MinPayloadLength"), MinPayloadLengthFromIni) ||
+		ConfigFile.GetInt64(TEXT("Core.ContentVirtualization"), TEXT("MinPayloadLength"), MinPayloadLengthFromIni))
 	{
 		MinPayloadLength = MinPayloadLengthFromIni;
 		UE_LOG(LogVirtualization, Display, TEXT("\tMinPayloadLength : %" INT64_FMT), MinPayloadLength );
@@ -793,8 +793,8 @@ void FVirtualizationManager::ApplySettingsFromConfigFiles(const FConfigFile& Con
 	}
 
 	FString BackendGraphNameFromIni;
-	if (ConfigFile.GetString(TEXT("Core.ContentVirtualization"), TEXT("BackendGraph"), BackendGraphNameFromIni) ||
-		ConfigFile.GetString(TEXT("Core.VirtualizationModule"), TEXT("BackendGraph"), BackendGraphNameFromIni))
+	if (ConfigFile.GetString(TEXT("Core.VirtualizationModule"), TEXT("BackendGraph"), BackendGraphNameFromIni) ||
+		ConfigFile.GetString(TEXT("Core.ContentVirtualization"), TEXT("BackendGraph"), BackendGraphNameFromIni))
 	{
 		BackendGraphName = BackendGraphNameFromIni;
 		UE_LOG(LogVirtualization, Display, TEXT("\tBackendGraphName : %s"), *BackendGraphName );
@@ -816,8 +816,8 @@ void FVirtualizationManager::ApplySettingsFromConfigFiles(const FConfigFile& Con
 	}
 
 	FString FilterModeFromIni;
-	if (ConfigFile.GetString(TEXT("Core.ContentVirtualization"), TEXT("FilterMode"), FilterModeFromIni) ||
-		ConfigFile.GetString(TEXT("Core.VirtualizationModule"), TEXT("FilterMode"), FilterModeFromIni))
+	if (ConfigFile.GetString(TEXT("Core.VirtualizationModule"), TEXT("FilterMode"), FilterModeFromIni) ||
+		ConfigFile.GetString(TEXT("Core.ContentVirtualization"), TEXT("FilterMode"), FilterModeFromIni))
 	{
 		if(LexTryParseString(FilteringMode, FilterModeFromIni))
 		{
@@ -834,8 +834,8 @@ void FVirtualizationManager::ApplySettingsFromConfigFiles(const FConfigFile& Con
 	}
 
 	bool bFilterEngineContentFromIni = true;
-	if (ConfigFile.GetBool(TEXT("Core.ContentVirtualization"), TEXT("FilterEngineContent"), bFilterEngineContentFromIni) ||
-		ConfigFile.GetBool(TEXT("Core.VirtualizationModule"), TEXT("FilterEngineContent"), bFilterEngineContentFromIni))
+	if (ConfigFile.GetBool(TEXT("Core.VirtualizationModule"), TEXT("FilterEngineContent"), bFilterEngineContentFromIni) ||
+		ConfigFile.GetBool(TEXT("Core.ContentVirtualization"), TEXT("FilterEngineContent"), bFilterEngineContentFromIni))
 	{
 		bFilterEngineContent = bFilterEngineContentFromIni;
 		UE_LOG(LogVirtualization, Display, TEXT("\tFilterEngineContent : %s"), bFilterEngineContent ? TEXT("true") : TEXT("false"));
@@ -846,8 +846,8 @@ void FVirtualizationManager::ApplySettingsFromConfigFiles(const FConfigFile& Con
 	}
 	
 	bool bFilterEnginePluginContentFromIni = true;
-	if (ConfigFile.GetBool(TEXT("Core.ContentVirtualization"), TEXT("FilterEnginePluginContent"), bFilterEnginePluginContentFromIni) ||
-		ConfigFile.GetBool(TEXT("Core.VirtualizationModule"), TEXT("FilterEnginePluginContent"), bFilterEnginePluginContentFromIni))
+	if (ConfigFile.GetBool(TEXT("Core.VirtualizationModule"), TEXT("FilterEnginePluginContent"), bFilterEnginePluginContentFromIni) ||
+		ConfigFile.GetBool(TEXT("Core.ContentVirtualization"), TEXT("FilterEnginePluginContent"), bFilterEnginePluginContentFromIni))
 	{
 		bFilterEnginePluginContent = bFilterEnginePluginContentFromIni;
 		UE_LOG(LogVirtualization, Display, TEXT("\tFilterEnginePluginContent : %s"), bFilterEnginePluginContent ? TEXT("true") : TEXT("false"));
@@ -871,8 +871,8 @@ void FVirtualizationManager::ApplySettingsFromConfigFiles(const FConfigFile& Con
 
 	// Optional
 	TArray<FString> DisabledAssetTypesFromIni;
-	if (ConfigFile.GetArray(TEXT("Core.ContentVirtualization"), TEXT("DisabledAsset"), DisabledAssetTypesFromIni) > 0 ||
-		ConfigFile.GetArray(TEXT("Core.VirtualizationModule"), TEXT("DisabledAsset"), DisabledAssetTypesFromIni) > 0)
+	if (ConfigFile.GetArray(TEXT("Core.VirtualizationModule"), TEXT("DisabledAsset"), DisabledAssetTypesFromIni) > 0 ||
+		ConfigFile.GetArray(TEXT("Core.ContentVirtualization"), TEXT("DisabledAsset"), DisabledAssetTypesFromIni) > 0)
 	{
 		UE_LOG(LogVirtualization, Display, TEXT("\tVirtualization is disabled for payloads of the following assets:"));
 		DisabledAssetTypes.Reserve(DisabledAssetTypesFromIni.Num());
