Add FSourceControlBackend support for storing virtualized payloads in a stream type depot.

#rnx

- Based on work from Jess.Kube

### Virtualization
- The source control backend now takes an optional config option "ClientStream" which takes the name of the client stream to use.
- If a client stream is set then workspaces created for payload submission will use that and not provide a client-view mapping.

### PerforceSourceControl
- Allow FCreateWorkspace to create workspaces with streams as well as classic workspaces cia FCreateWorkspace::SetStream.
- Add a method FCreateWorkspace::ClearClientViewMappings which will clear any client view mappings already added to the operation.
- If we detect that a FCreateWorkspace operation has both a stream set and client view mappings set then FPerforceCreateWorkspaceWorker will return an error. Perforce will allow us to creat a client spec with both entries, but will default to using the stream. Technically we could allow this too but it might cause unexpected behaviour to the caller. It is better to give a clear error and fail the workspace creation.

diff --git a/Engine/Plugins/Developer/PerforceSourceControl/Source/PerforceSourceControl/Private/PerforceSourceControlOperations.cpp b/Engine/Plugins/Developer/PerforceSourceControl/Source/PerforceSourceControl/Private/PerforceSourceControlOperations.cpp
--- a/Engine/Plugins/Developer/PerforceSourceControl/Source/PerforceSourceControl/Private/PerforceSourceControlOperations.cpp
+++ b/Engine/Plugins/Developer/PerforceSourceControl/Source/PerforceSourceControl/Private/PerforceSourceControlOperations.cpp
@@ -3138,6 +3138,19 @@ bool FPerforceCreateWorkspaceWorker::Execute(FPerforceSourceControlCommand& InCo
 		FPerforceConnection& Connection = ScopedConnection.GetConnection();
 
 		TSharedRef<FCreateWorkspace, ESPMode::ThreadSafe> Operation = StaticCastSharedRef<FCreateWorkspace>(InCommand.Operation);
+	
+		// Early out if the operation set up is invalid
+		// Note if we set Stream and View this will not result in a p4 error and the View setting will be ignored.
+		// As a precaution we will consider this an error instead to avoid unintended consequences for the caller.
+		if (!Operation->GetWorkspaceStream().IsEmpty() && !Operation->GetClientView().IsEmpty())
+		{
+			FText ErrorMsg = LOCTEXT("WorkspaceInvalidSetup", "Cannot create a workspace for a stream and apply a client view at the same time");
+
+			InCommand.ResultInfo.ErrorMessages.Add(ErrorMsg);
+			InCommand.bCommandSuccessful = false;
+			
+			return InCommand.bCommandSuccessful;
+		}
 
 		TStringBuilder<2048> ClientDesc;
 
@@ -3146,17 +3159,28 @@ bool FPerforceCreateWorkspaceWorker::Execute(FPerforceSourceControlCommand& InCo
 
 		if (!Operation->GetWorkspaceDescription().IsEmpty())
 		{
-			ClientDesc << TEXT("Description:\t") << Operation->GetWorkspaceDescription() << TEXT("\n\n");
+			ClientDesc << TEXT("Description:\n\t") << Operation->GetWorkspaceDescription() << TEXT("\n\n");
 		}
 
 		ClientDesc << TEXT("Root:\t") << Operation->GetWorkspaceRoot() << TEXT("\n\n");
-		ClientDesc << TEXT("View:\t\n\n");
 
-		for (const FCreateWorkspace::FClientViewMapping& Mapping : Operation->GetClientView())
+		if (!Operation->GetWorkspaceStream().IsEmpty())
+		{
+			ClientDesc << TEXT("Stream:\t") << Operation->GetWorkspaceStream() << TEXT("\n\n");
+		}
+
+		if(!Operation->GetClientView().IsEmpty())
 		{
-			ClientDesc << TEXT("\t") << Mapping.Key;
-			ClientDesc << TEXT(" ");
-			ClientDesc << Mapping.Value << TEXT("\n");
+			ClientDesc << TEXT("View:\t\n");
+
+			for (const FCreateWorkspace::FClientViewMapping& Mapping : Operation->GetClientView())
+			{
+				ClientDesc << TEXT("\t") << Mapping.Key;
+				ClientDesc << TEXT(" ");
+				ClientDesc << Mapping.Value << TEXT("\n");
+			}
+
+			ClientDesc << TEXT("\n");
 		}
 
 		AddType(*Operation, ClientDesc);
diff --git a/Engine/Source/Developer/SourceControl/Public/SourceControlOperations.h b/Engine/Source/Developer/SourceControl/Public/SourceControlOperations.h
--- a/Engine/Source/Developer/SourceControl/Public/SourceControlOperations.h
+++ b/Engine/Source/Developer/SourceControl/Public/SourceControlOperations.h
@@ -1000,6 +1000,12 @@ public:
 		WorkspaceDescription = Desciption;
 	}
 
+	/** Set the stream to be used by the workspace, if left unset then a classic depot with ClientView will be used */
+	void SetStream(FStringView Stream)
+	{
+		WorkspaceStream = Stream;
+	}
+
 	/**
 	 * Add a new mapping for the client spec in the native format of the current source control provider.
 	 * These will be written to the client spec in the order that they are added.
@@ -1012,6 +1018,12 @@ public:
 		ClientView.Emplace(DepotPath, ClientPath);
 	}
 
+	/** Remove all currently set client-view mappings */
+	void ClearClientViewMappings()
+	{
+		ClientView.Empty();
+	}
+
 	/** Set the type of workspace to be created. @see FCreateWorkspace::EType */
 	void SetType(EType InType)
 	{
@@ -1033,6 +1045,11 @@ public:
 		return WorkspaceDescription;
 	}
 
+	const FString& GetWorkspaceStream() const
+	{
+		return WorkspaceStream;
+	}
+
 	const TArray<FClientViewMapping> GetClientView() const
 	{
 		return ClientView;
@@ -1061,6 +1078,7 @@ protected:
 	FString WorkspaceRoot;
 
 	FString WorkspaceDescription;
+	FString WorkspaceStream;
 
 	EType Type = EType::Writeable;
 
diff --git a/Engine/Source/Developer/Virtualization/Private/VirtualizationSourceControlBackend.cpp b/Engine/Source/Developer/Virtualization/Private/VirtualizationSourceControlBackend.cpp
--- a/Engine/Source/Developer/Virtualization/Private/VirtualizationSourceControlBackend.cpp
+++ b/Engine/Source/Developer/Virtualization/Private/VirtualizationSourceControlBackend.cpp
@@ -436,13 +436,20 @@ bool FSourceControlBackend::PushData(TArrayView<FPushRequest> Requests)
 		TRACE_CPUPROFILER_EVENT_SCOPE(FSourceControlBackend::PushData::CreateWorkspace);
 		TSharedRef<FCreateWorkspace> CreateWorkspaceCommand = ISourceControlOperation::Create<FCreateWorkspace>(WorkspaceName, SessionDirectory);
 
-		TStringBuilder<512> DepotMapping;
-		DepotMapping << DepotRoot << TEXT("...");
+		if (ClientStream.IsEmpty())
+		{
+			TStringBuilder<512> DepotMapping;
+			DepotMapping << DepotRoot << TEXT("...");
 
-		TStringBuilder<128> ClientMapping;
-		ClientMapping << TEXT("//") << WorkspaceName << TEXT("/...");
+			TStringBuilder<128> ClientMapping;
+			ClientMapping << TEXT("//") << WorkspaceName << TEXT("/...");
 
-		CreateWorkspaceCommand->AddNativeClientViewMapping(DepotMapping, ClientMapping);
+			CreateWorkspaceCommand->AddNativeClientViewMapping(DepotMapping, ClientMapping);
+		}
+		else
+		{
+			CreateWorkspaceCommand->SetStream(ClientStream);
+		}
 
 		if (bUsePartitionedClient)
 		{
@@ -453,7 +460,7 @@ bool FSourceControlBackend::PushData(TArrayView<FPushRequest> Requests)
 
 		if (SCCProvider->Execute(CreateWorkspaceCommand) != ECommandResult::Succeeded)
 		{
-			UE_LOG(LogVirtualization, Error, TEXT("[%s] Failed to create temp workspace '%s' to submit payloads from"),
+			UE_LOG(LogVirtualization, Error, TEXT("[%s] Failed to create temp workspace '%s' to use for payload submission"),
 				*GetDebugName(),
 				WorkspaceName.ToString());
 
@@ -716,19 +723,34 @@ bool FSourceControlBackend::DoPayloadsExist(TArrayView<const FIoHash> PayloadIds
 bool FSourceControlBackend::TryApplySettingsFromConfigFiles(const FString& ConfigEntry)
 {
 	// We require that a valid depot root has been provided
-	if (!FParse::Value(*ConfigEntry, TEXT("DepotRoot="), DepotRoot))
 	{
-		UE_LOG(LogVirtualization, Error, TEXT("'DepotRoot=' not found in the config file"));
-		return false;
-	}
+		if (!FParse::Value(*ConfigEntry, TEXT("DepotRoot="), DepotRoot))
+		{
+			UE_LOG(LogVirtualization, Error, TEXT("'DepotRoot=' not found in the config file"));
+			return false;
+		}
 
-	if (!DepotRoot.EndsWith(TEXT("/")))
-	{
-		DepotRoot.AppendChar(TEXT('/'));
+		if (!DepotRoot.EndsWith(TEXT("/")))
+		{
+			DepotRoot.AppendChar(TEXT('/'));
+		}
 	}
 
 	// Now parse the optional config values
 
+	// If the depot root is within a perforce stream then we must specify which stream. This may be a virtual stream with a custom view.
+	{
+		FParse::Value(*ConfigEntry, TEXT("ClientStream="), ClientStream);
+		if (!ClientStream.IsEmpty())
+		{
+			UE_LOG(LogVirtualization, Log, TEXT("[%s] Using client stream: '%s'"), *GetDebugName(), *ClientStream);
+		}
+		else
+		{
+			UE_LOG(LogVirtualization, Log, TEXT("[%s] Not using client stream"), *GetDebugName());
+		}
+	}
+
 	// Check to see if we should use partitioned clients or not. This is a perforce specific optimization to make the workspace churn cheaper on the server
 	{
 		FParse::Bool(*ConfigEntry, TEXT("UsePartitionedClient="), bUsePartitionedClient);
diff --git a/Engine/Source/Developer/Virtualization/Private/VirtualizationSourceControlBackend.h b/Engine/Source/Developer/Virtualization/Private/VirtualizationSourceControlBackend.h
--- a/Engine/Source/Developer/Virtualization/Private/VirtualizationSourceControlBackend.h
+++ b/Engine/Source/Developer/Virtualization/Private/VirtualizationSourceControlBackend.h
@@ -26,6 +26,8 @@ class FSemaphore;
  * depot where the payload files are being stored.
  * 
  * Optional Values:
+ * ClientStream [string]:		Used when the payloads are stored in a stream based depot. It should contain
+ *								the stream name to use when creating a workspace for payload submission.
  * UsePartitionedClient [bool]:	When true the temporary workspace client created to submit payloads 
  *								from will be created as a partitioned workspace which is less overhead
  *								on the source control server. If your server does not support this then
@@ -93,6 +95,9 @@ private:
 	/** The root where the virtualized payloads are stored in source control */
 	FString DepotRoot;
 
+	/** The stream containing the DepotRoot where the virtualized payloads are stored in source control */
+	FString ClientStream;
+
 	/** The root directory from which payloads are submitted. */
 	FString SubmissionRootDir;
 
