Renamed the config value [Core.VirtualizationModule].EnablePushToBackend to [Core.VirtualizationModule].EnablePayloadVirtualization so that it makes more sense

#rnx

- The older versions of this value ([Core.ContentVirtualization].EnablePushToBackend and [Core.VirtualizationModule].EnablePushToBackend) will continue to work but will log a warning to the user telling them to update it.
- It is very unlikely that anyone is actually using these old values and so we can remove this backwards compatibility when 5.1 ships. It is only present to allow us to fix up internal projects and to avoid interrupting people testing the system.

diff --git a/Engine/Config/BaseEngine.ini b/Engine/Config/BaseEngine.ini
--- a/Engine/Config/BaseEngine.ini
+++ b/Engine/Config/BaseEngine.ini
@@ -1237,7 +1237,7 @@ AllowVirtualizationOptOut =False
 SystemName=None
 
 [Core.VirtualizationModule]
-EnablePushToBackend=False
+EnablePayloadVirtualization=True
 EnableCacheAfterPull=True
 MinPayloadLength=0
 BackendGraph=ContentVirtualizationBackendGraph_None
diff --git a/Engine/Source/Developer/Virtualization/Private/VirtualizationManager.cpp b/Engine/Source/Developer/Virtualization/Private/VirtualizationManager.cpp
--- a/Engine/Source/Developer/Virtualization/Private/VirtualizationManager.cpp
+++ b/Engine/Source/Developer/Virtualization/Private/VirtualizationManager.cpp
@@ -260,7 +260,7 @@ namespace Profiling
 } //namespace Profiling
 
 FVirtualizationManager::FVirtualizationManager()
-	: bEnablePayloadPushing(true)
+	: bEnablePayloadVirtualization(true)
 	, bEnableCacheAfterPull(true)
 	, MinPayloadLength(0)
 	, BackendGraphName(TEXT("ContentVirtualizationBackendGraph_None"))
@@ -355,7 +355,7 @@ bool FVirtualizationManager::IsEnabled() const
 
 bool FVirtualizationManager::IsPushingEnabled(EStorageType StorageType) const
 {
-	if (!bEnablePayloadPushing)
+	if (!bEnablePayloadVirtualization)
 	{
 		return false;
 	}
@@ -478,7 +478,7 @@ bool FVirtualizationManager::PushData(TArrayView<FPushRequest> Requests, EStorag
 	}
 
 	// Early out if there are no backends
-	if (!IsEnabled() || bEnablePayloadPushing == false)
+	if (!IsEnabled() || bEnablePayloadVirtualization == false)
 	{
 		return false;
 	}
@@ -753,24 +753,48 @@ void FVirtualizationManager::ApplySettingsFromConfigFiles(const FConfigFile& Con
 {
 	UE_LOG(LogVirtualization, Display, TEXT("Loading virtualization manager settings from config files..."));
 
+	const TCHAR* LegacyConfigSection = TEXT("Core.ContentVirtualization");
+	const TCHAR* ConfigSection = TEXT("Core.VirtualizationModule");
+
 	// Note that all options are doubled up as we are moving the options for this module from "Core.ContentVirtualization"
 	// to it's own specific "Core.VirtualizationModule" section. This duplication can be removed before we ship 5.1
-
-	bool bEnablePayloadPushingFromIni = false;
-	if (ConfigFile.GetBool(TEXT("Core.ContentVirtualization"), TEXT("EnablePushToBackend"), bEnablePayloadPushingFromIni) ||
-		ConfigFile.GetBool(TEXT("Core.VirtualizationModule"), TEXT("EnablePushToBackend"), bEnablePayloadPushingFromIni))
-	{
-		bEnablePayloadPushing = bEnablePayloadPushingFromIni;
-		UE_LOG(LogVirtualization, Display, TEXT("\tEnablePushToBackend : %s"), bEnablePayloadPushing ? TEXT("true") : TEXT("false") );
-	}
-	else
+	
 	{
-		UE_LOG(LogVirtualization, Error, TEXT("Failed to load [Core.VirtualizationModule].EnablePushToBackend from config file!"));
+		// This value was moved from Core.ContentVirtualization to Core.VirtualizationModule then renamed from
+		// 'EnablePushToBackend' to 'EnablePayloadVirtualization' so there are a few paths we need to cover here.
+		// This can also be cleaned up for 5.1 shipping.
+		bool bLoadedFromFile = false;
+		bool bEnablePayloadVirtualizationFromIni = false;
+		
+		if (ConfigFile.GetBool(LegacyConfigSection, TEXT("EnablePushToBackend"), bEnablePayloadVirtualizationFromIni))
+		{
+			UE_LOG(LogVirtualization, Warning, TEXT("\tFound legacy ini file setting [Core.ContentVirtualization].EnablePushToBackend, rename to [Core.VirtualizationModule].EnablePayloadVirtualization"));
+			bLoadedFromFile = true;
+		}
+		else if (ConfigFile.GetBool(ConfigSection, TEXT("EnablePushToBackend"), bEnablePayloadVirtualizationFromIni))
+		{
+			UE_LOG(LogVirtualization, Warning, TEXT("\tFound legacy ini file setting [Core.VirtualizationModule].EnablePushToBackend, rename to [Core.VirtualizationModule].EnablePayloadVirtualization"));
+			bLoadedFromFile = true;
+		}
+		else if (ConfigFile.GetBool(ConfigSection, TEXT("EnablePayloadVirtualization"), bEnablePayloadVirtualizationFromIni))
+		{
+			bLoadedFromFile = true;
+		}
+
+		if (bLoadedFromFile)
+		{
+			bEnablePayloadVirtualization = bEnablePayloadVirtualizationFromIni;
+			UE_LOG(LogVirtualization, Display, TEXT("\tEnablePayloadVirtualization : %s"), bEnablePayloadVirtualization ? TEXT("true") : TEXT("false"));
+		}
+		else
+		{
+			UE_LOG(LogVirtualization, Error, TEXT("Failed to load [Core.VirtualizationModule].EnablePayloadVirtualization from config file!"));
+		}
 	}
 
 	bool bEnableCacheAfterPullFromIni = false;
-	if (ConfigFile.GetBool(TEXT("Core.ContentVirtualization"), TEXT("EnableCacheAfterPull"), bEnableCacheAfterPullFromIni) ||
-		ConfigFile.GetBool(TEXT("Core.VirtualizationModule"), TEXT("EnableCacheAfterPull"), bEnableCacheAfterPullFromIni))
+	if (ConfigFile.GetBool(LegacyConfigSection, TEXT("EnableCacheAfterPull"), bEnableCacheAfterPullFromIni) ||
+		ConfigFile.GetBool(ConfigSection, TEXT("EnableCacheAfterPull"), bEnableCacheAfterPullFromIni))
 	{
 		bEnableCacheAfterPull = bEnableCacheAfterPullFromIni;
 		UE_LOG(LogVirtualization, Display, TEXT("\tCachePulledPayloads : %s"), bEnableCacheAfterPull ? TEXT("true") : TEXT("false"));
@@ -781,8 +805,8 @@ void FVirtualizationManager::ApplySettingsFromConfigFiles(const FConfigFile& Con
 	}
 
 	int64 MinPayloadLengthFromIni = 0;
-	if (ConfigFile.GetInt64(TEXT("Core.ContentVirtualization"), TEXT("MinPayloadLength"), MinPayloadLengthFromIni) ||
-		ConfigFile.GetInt64(TEXT("Core.VirtualizationModule"), TEXT("MinPayloadLength"), MinPayloadLengthFromIni))
+	if (ConfigFile.GetInt64(LegacyConfigSection, TEXT("MinPayloadLength"), MinPayloadLengthFromIni) ||
+		ConfigFile.GetInt64(ConfigSection, TEXT("MinPayloadLength"), MinPayloadLengthFromIni))
 	{
 		MinPayloadLength = MinPayloadLengthFromIni;
 		UE_LOG(LogVirtualization, Display, TEXT("\tMinPayloadLength : %" INT64_FMT), MinPayloadLength );
@@ -793,8 +817,8 @@ void FVirtualizationManager::ApplySettingsFromConfigFiles(const FConfigFile& Con
 	}
 
 	FString BackendGraphNameFromIni;
-	if (ConfigFile.GetString(TEXT("Core.ContentVirtualization"), TEXT("BackendGraph"), BackendGraphNameFromIni) ||
-		ConfigFile.GetString(TEXT("Core.VirtualizationModule"), TEXT("BackendGraph"), BackendGraphNameFromIni))
+	if (ConfigFile.GetString(LegacyConfigSection, TEXT("BackendGraph"), BackendGraphNameFromIni) ||
+		ConfigFile.GetString(ConfigSection, TEXT("BackendGraph"), BackendGraphNameFromIni))
 	{
 		BackendGraphName = BackendGraphNameFromIni;
 		UE_LOG(LogVirtualization, Display, TEXT("\tBackendGraphName : %s"), *BackendGraphName );
@@ -805,7 +829,7 @@ void FVirtualizationManager::ApplySettingsFromConfigFiles(const FConfigFile& Con
 	}
 
 	FString VirtualizationProcessTagFromIni;
-	if (ConfigFile.GetString(TEXT("Core.VirtualizationModule"), TEXT("VirtualizationProcessTag"), VirtualizationProcessTagFromIni))
+	if (ConfigFile.GetString(ConfigSection, TEXT("VirtualizationProcessTag"), VirtualizationProcessTagFromIni))
 	{
 		VirtualizationProcessTag = VirtualizationProcessTagFromIni;
 		UE_LOG(LogVirtualization, Display, TEXT("\tVirtualizationProcessTag : %s"), *VirtualizationProcessTag);
@@ -816,8 +840,8 @@ void FVirtualizationManager::ApplySettingsFromConfigFiles(const FConfigFile& Con
 	}
 
 	FString FilterModeFromIni;
-	if (ConfigFile.GetString(TEXT("Core.ContentVirtualization"), TEXT("FilterMode"), FilterModeFromIni) ||
-		ConfigFile.GetString(TEXT("Core.VirtualizationModule"), TEXT("FilterMode"), FilterModeFromIni))
+	if (ConfigFile.GetString(LegacyConfigSection, TEXT("FilterMode"), FilterModeFromIni) ||
+		ConfigFile.GetString(ConfigSection, TEXT("FilterMode"), FilterModeFromIni))
 	{
 		if(LexTryParseString(FilteringMode, FilterModeFromIni))
 		{
@@ -834,8 +858,8 @@ void FVirtualizationManager::ApplySettingsFromConfigFiles(const FConfigFile& Con
 	}
 
 	bool bFilterEngineContentFromIni = true;
-	if (ConfigFile.GetBool(TEXT("Core.ContentVirtualization"), TEXT("FilterEngineContent"), bFilterEngineContentFromIni) ||
-		ConfigFile.GetBool(TEXT("Core.VirtualizationModule"), TEXT("FilterEngineContent"), bFilterEngineContentFromIni))
+	if (ConfigFile.GetBool(LegacyConfigSection, TEXT("FilterEngineContent"), bFilterEngineContentFromIni) ||
+		ConfigFile.GetBool(ConfigSection, TEXT("FilterEngineContent"), bFilterEngineContentFromIni))
 	{
 		bFilterEngineContent = bFilterEngineContentFromIni;
 		UE_LOG(LogVirtualization, Display, TEXT("\tFilterEngineContent : %s"), bFilterEngineContent ? TEXT("true") : TEXT("false"));
@@ -846,8 +870,8 @@ void FVirtualizationManager::ApplySettingsFromConfigFiles(const FConfigFile& Con
 	}
 	
 	bool bFilterEnginePluginContentFromIni = true;
-	if (ConfigFile.GetBool(TEXT("Core.ContentVirtualization"), TEXT("FilterEnginePluginContent"), bFilterEnginePluginContentFromIni) ||
-		ConfigFile.GetBool(TEXT("Core.VirtualizationModule"), TEXT("FilterEnginePluginContent"), bFilterEnginePluginContentFromIni))
+	if (ConfigFile.GetBool(LegacyConfigSection, TEXT("FilterEnginePluginContent"), bFilterEnginePluginContentFromIni) ||
+		ConfigFile.GetBool(ConfigSection, TEXT("FilterEnginePluginContent"), bFilterEnginePluginContentFromIni))
 	{
 		bFilterEnginePluginContent = bFilterEnginePluginContentFromIni;
 		UE_LOG(LogVirtualization, Display, TEXT("\tFilterEnginePluginContent : %s"), bFilterEnginePluginContent ? TEXT("true") : TEXT("false"));
@@ -859,7 +883,7 @@ void FVirtualizationManager::ApplySettingsFromConfigFiles(const FConfigFile& Con
 
 	// Optional
 	bool bFilterMapContentFromIni = false;
-	if (ConfigFile.GetBool(TEXT("Core.VirtualizationModule"), TEXT("FilterMapContent"), bFilterMapContentFromIni))
+	if (ConfigFile.GetBool(ConfigSection, TEXT("FilterMapContent"), bFilterMapContentFromIni))
 	{
 		bFilterMapContent = bFilterMapContentFromIni;
 		UE_LOG(LogVirtualization, Display, TEXT("\tFilterMapContent : %s"), bFilterMapContent ? TEXT("true") : TEXT("false"));
@@ -871,8 +895,8 @@ void FVirtualizationManager::ApplySettingsFromConfigFiles(const FConfigFile& Con
 
 	// Optional
 	TArray<FString> DisabledAssetTypesFromIni;
-	if (ConfigFile.GetArray(TEXT("Core.ContentVirtualization"), TEXT("DisabledAsset"), DisabledAssetTypesFromIni) > 0 ||
-		ConfigFile.GetArray(TEXT("Core.VirtualizationModule"), TEXT("DisabledAsset"), DisabledAssetTypesFromIni) > 0)
+	if (ConfigFile.GetArray(LegacyConfigSection, TEXT("DisabledAsset"), DisabledAssetTypesFromIni) > 0 ||
+		ConfigFile.GetArray(ConfigSection, TEXT("DisabledAsset"), DisabledAssetTypesFromIni) > 0)
 	{
 		UE_LOG(LogVirtualization, Display, TEXT("\tVirtualization is disabled for payloads of the following assets:"));
 		DisabledAssetTypes.Reserve(DisabledAssetTypesFromIni.Num());
@@ -885,7 +909,7 @@ void FVirtualizationManager::ApplySettingsFromConfigFiles(const FConfigFile& Con
 
 	// Optional
 	bool bAllowSubmitIfVirtualizationFailedFromIni = true;
-	if (ConfigFile.GetBool(TEXT("Core.VirtualizationModule"), TEXT("AllowSubmitIfVirtualizationFailed"), bAllowSubmitIfVirtualizationFailedFromIni))
+	if (ConfigFile.GetBool(ConfigSection, TEXT("AllowSubmitIfVirtualizationFailed"), bAllowSubmitIfVirtualizationFailedFromIni))
 	{
 		bAllowSubmitIfVirtualizationFailed = bAllowSubmitIfVirtualizationFailedFromIni;
 		UE_LOG(LogVirtualization, Display, TEXT("\tAllowSubmitIfVirtualizationFailed : %s"), bAllowSubmitIfVirtualizationFailed ? TEXT("true") : TEXT("false"));
@@ -896,7 +920,7 @@ void FVirtualizationManager::ApplySettingsFromConfigFiles(const FConfigFile& Con
 	}
 
 	// Check for any legacy settings and print them out (easier to do this in one block rather than one and time)
-	if (const FConfigSection* LegacySection = ConfigFile.Find(TEXT("Core.ContentVirtualization")))
+	if (const FConfigSection* LegacySection = ConfigFile.Find(LegacyConfigSection))
 	{
 		if (LegacySection->Num() > 1)
 		{
diff --git a/Engine/Source/Developer/Virtualization/Private/VirtualizationManager.h b/Engine/Source/Developer/Virtualization/Private/VirtualizationManager.h
--- a/Engine/Source/Developer/Virtualization/Private/VirtualizationManager.h
+++ b/Engine/Source/Developer/Virtualization/Private/VirtualizationManager.h
@@ -218,7 +218,7 @@ private:
 	// The following members are set from the config file
 
 	/** Are payloads allowed to be virtualized. Defaults to true. */
-	bool bEnablePayloadPushing;
+	bool bEnablePayloadVirtualization;
 
 	/** Should payloads be cached locally after being pulled from persistent storage? Defaults to true. */
 	bool bEnableCacheAfterPull;
