Change the VA graph name LocalStorageHierarchy to 'CacheStorageHierarchy' which better describes it.

#rnx

- The orignal name implied tha the payloads would only be stored locally on the users machine, but the intent is to describe a number of backends that are a) faster than the persistent storage backends b) no guarantee that the payload will be found.
- The new name better describes this functionality.
- We still accept 'LocalStorageHierarchy' but will log a warning informing the user to update their graph but this backwards compatibility code will most likely be removed before 5.1 ships.
-- EStorageType::Local will follow the usual deprecation rules.
- Not strictly related to this change but I also cleaned up and added additional loggng during the initialization of the virtualization manager.

diff --git a/Engine/Config/BaseEngine.ini b/Engine/Config/BaseEngine.ini
--- a/Engine/Config/BaseEngine.ini
+++ b/Engine/Config/BaseEngine.ini
@@ -1256,18 +1256,18 @@ FilterMapContent=True
 
 ; A backend graph with no backends
 [ContentVirtualizationBackendGraph_None]
-LocalStorageHierarchy=()
+CacheStorageHierarchy=()
 PersistentStorageHierarchy=()
 
 ; A simple backend graph that will store virtualized data locally, suitable for testing.
 [ContentVirtualizationBackendGraph_LocalCache]
-LocalStorageHierarchy=(Entry=LocalFileCache)
+CacheStorageHierarchy=(Entry=LocalFileCache)
 PersistentStorageHierarchy=()
 LocalFileCache=(Type=FileSystem, Path="%GAMESAVEDDIR%VirtualizedPayloads/")
 
 ; A backend graph that will connect to a locally hosted Jupiter service along with a local file cache, suitable for testing.
 [ContentVirtualizationBackendGraph_JupiterLocalHost]
-LocalStorageHierarchy=(Entry=LocalFileCache)
+CacheStorageHierarchy=(Entry=LocalFileCache)
 PersistentStorageHierarchy=(JupiterCache)
 LocalFileCache=(Type=FileSystem, Path="%GAMESAVEDDIR%VirtualizedPayloads/")
 JupiterCache=(Type=Jupiter, Host="http://localhost", Namespace="ue.mirage", ChunkSize=524288)
diff --git a/Engine/Source/Developer/Virtualization/Private/VirtualizationManager.cpp b/Engine/Source/Developer/Virtualization/Private/VirtualizationManager.cpp
--- a/Engine/Source/Developer/Virtualization/Private/VirtualizationManager.cpp
+++ b/Engine/Source/Developer/Virtualization/Private/VirtualizationManager.cpp
@@ -122,6 +122,16 @@ TArray<FString> ParseEntries(const FString& Data)
 	return Entries;
 }
 
+/** 
+ * Utility to make 'StorageType ==  EStorageType::Cache' checks easier while EStorageType::Local continues to exist.
+ * When the deprecated value is removed this can also be removed and code calling it can just check for EStorageType::Cache */
+bool IsCacheType(EStorageType StorageType)
+{
+PRAGMA_DISABLE_DEPRECATION_WARNINGS
+	return StorageType == EStorageType::Local || StorageType ==  EStorageType::Cache;
+PRAGMA_ENABLE_DEPRECATION_WARNINGS
+}
+
 /**
  * Profiling data allowing us to track how payloads are being push/pulled during the lifespan of the process. Note that as all backends are
  * created at the same time, we don't need to add locked when accessing the maps. In addition FCookStats is thread safe when adding hits/misses
@@ -282,7 +292,7 @@ FVirtualizationManager::~FVirtualizationManager()
 
 	UE_LOG(LogVirtualization, Log, TEXT("Destroying backends"));
 
-	LocalCachableBackends.Empty();
+	CacheStorageBackends.Empty();
 	PersistentStorageBackends.Empty();
 	PullEnabledBackends.Empty();
 
@@ -326,6 +336,8 @@ bool FVirtualizationManager::Initialize(const FInitParams& InitParams)
 {
 	TRACE_CPUPROFILER_EVENT_SCOPE(FVirtualizationManager::Initialize);
 
+	UE_LOG(LogVirtualization, Display, TEXT("Initializing the virtualization manager..."));
+
 	// TODO: Ideally we'd break this down further, or at least have a FScopedSlowTask for each
 	// backend initialization but the slow task system will only update the UI every 0.2 seconds
 	// so if we have too many small tasks we might show misleading data to the user, so it is 
@@ -345,6 +357,8 @@ bool FVirtualizationManager::Initialize(const FInitParams& InitParams)
 
 	MountBackends(InitParams.ConfigFile);
 
+	UE_LOG(LogVirtualization, Display, TEXT("Virtualization manager initialization completed"));
+
 	return true;
 }
 
@@ -362,8 +376,11 @@ bool FVirtualizationManager::IsPushingEnabled(EStorageType StorageType) const
 
 	switch (StorageType)
 	{
+PRAGMA_DISABLE_DEPRECATION_WARNINGS
 		case EStorageType::Local:
-			return !LocalCachableBackends.IsEmpty();
+PRAGMA_ENABLE_DEPRECATION_WARNINGS
+		case EStorageType::Cache:
+			return !CacheStorageBackends.IsEmpty();
 			break;
 
 		case EStorageType::Persistent:
@@ -491,7 +508,7 @@ bool FVirtualizationManager::PushData(TArrayView<FPushRequest> Requests, EStorag
 
 	int32 ErrorCount = 0;
 	bool bWasPayloadPushed = false;
-	FBackendArray& Backends = StorageType == EStorageType::Local ? LocalCachableBackends : PersistentStorageBackends;
+	FBackendArray& Backends = IsCacheType(StorageType) ? CacheStorageBackends : PersistentStorageBackends;
 
 	for (IVirtualizationBackend* Backend : Backends)
 	{
@@ -533,7 +550,7 @@ bool FVirtualizationManager::PushData(TArrayView<FPushRequest> Requests, EStorag
 
 	// For local storage we consider the push to have failed only if ALL backends gave an error, if at least one backend succeeded then the operation succeeded.
 	// For persistent storage we require that all backends succeeded, so any errors will fail the push operation.
-	return StorageType == EStorageType::Local ? ErrorCount < Backends.Num() : ErrorCount == 0;
+	return IsCacheType(StorageType) ? ErrorCount < Backends.Num() : ErrorCount == 0;
 }
 
 FCompressedBuffer FVirtualizationManager::PullData(const FIoHash& Id)
@@ -597,7 +614,7 @@ EQueryResult FVirtualizationManager::QueryPayloadStatuses(TArrayView<const FIoHa
 		OutStatuses[Index] = Ids[Index].IsZero() ? EPayloadStatus::Invalid : EPayloadStatus::NotFound;
 	}
 
-	FBackendArray& Backends = StorageType == EStorageType::Local ? LocalCachableBackends : PersistentStorageBackends;
+	FBackendArray& Backends = IsCacheType(StorageType) ? CacheStorageBackends : PersistentStorageBackends;
 
 	TArray<int8> HitCount;
 	TArray<bool> Results;
@@ -1230,26 +1247,38 @@ void FVirtualizationManager::MountBackends(const FConfigFile& ConfigFile)
 
 	// It is important to parse the local storage hierarchy first so those backends will show up before the
 	// persistent storage backends in 'PullEnabledBackends'.
-	ParseHierarchy(ConfigFile, GraphName, TEXT("LocalStorageHierarchy"), FactoryLookupTable, LocalCachableBackends);
-	ParseHierarchy(ConfigFile, GraphName, TEXT("PersistentStorageHierarchy"), FactoryLookupTable, PersistentStorageBackends);
+	ParseHierarchy(ConfigFile, GraphName, TEXT("CacheStorageHierarchy"), TEXT("LocalStorageHierarchy"), FactoryLookupTable, CacheStorageBackends);
+	ParseHierarchy(ConfigFile, GraphName, TEXT("PersistentStorageHierarchy"), nullptr, FactoryLookupTable, PersistentStorageBackends);
 
 	// Apply and disabled backends from the command line
 	UpdateBackendDebugState();
 }
 
-void FVirtualizationManager::ParseHierarchy(const FConfigFile& ConfigFile, const TCHAR* GraphName, const TCHAR* HierarchyKey, const FRegistedFactories& FactoryLookupTable, FBackendArray& PushArray)
+void FVirtualizationManager::ParseHierarchy(const FConfigFile& ConfigFile, const TCHAR* GraphName, const TCHAR* HierarchyKey, const TCHAR* LegacyHierarchyKey, const FRegistedFactories& FactoryLookupTable, FBackendArray& PushArray)
 {
+	bool bFoundHierarchy = false;
 	FString HierarchyData;
-	if (ConfigFile.GetValue(GraphName, HierarchyKey, HierarchyData))
+
+	if (LegacyHierarchyKey != nullptr && ConfigFile.GetValue(GraphName, LegacyHierarchyKey, HierarchyData))
+	{
+		UE_LOG(LogVirtualization, Warning, TEXT("\tFound legacy HierarchyKey '%s', rename to '%s'"), LegacyHierarchyKey, HierarchyKey);
+		bFoundHierarchy = true;
+	}
+	else if (ConfigFile.GetValue(GraphName, HierarchyKey, HierarchyData))
+	{
+		bFoundHierarchy = true;
+	}
+
+	if (bFoundHierarchy)
 	{
 		if (HierarchyData.IsEmpty())
 		{
-			UE_LOG(LogVirtualization, Fatal, TEXT("The '%s' entry for backend graph '%s' is empty [ini=%s]."), HierarchyKey, GraphName, *GEngineIni);
+			UE_LOG(LogVirtualization, Fatal, TEXT("\tThe '%s' entry for backend graph '%s' is empty [ini=%s]."), HierarchyKey, GraphName, *GEngineIni);
 		}
 
 		const TArray<FString> Entries = ParseEntries(HierarchyData);
 
-		UE_LOG(LogVirtualization, Display, TEXT("'%s' has %d backend(s)"), HierarchyKey, Entries.Num());
+		UE_LOG(LogVirtualization, Display, TEXT("\tThe '%s' has %d backend(s)"), HierarchyKey, Entries.Num());
 
 		for (const FString& Entry : Entries)
 		{
@@ -1258,7 +1287,7 @@ void FVirtualizationManager::ParseHierarchy(const FConfigFile& ConfigFile, const
 	}
 	else
 	{
-		UE_LOG(LogVirtualization, Display, TEXT("No entries for '%s' in the content virtualization backend graph '%s' [ini=%s]."), HierarchyKey, GraphName, *GEngineIni);
+		UE_LOG(LogVirtualization, Display, TEXT("\tNo entries for '%s' in the content virtualization backend graph '%s' [ini=%s]."), HierarchyKey, GraphName, *GEngineIni);
 	}
 }
 
@@ -1348,7 +1377,7 @@ void FVirtualizationManager::CachePayload(const FIoHash& Id, const FCompressedBu
 	TRACE_CPUPROFILER_EVENT_SCOPE(FVirtualizationManager::CachePayload);
 
 	// We start caching at the first (assumed to be fastest) local cache backend. 
-	for (IVirtualizationBackend* BackendToCache : LocalCachableBackends)
+	for (IVirtualizationBackend* BackendToCache : CacheStorageBackends)
 	{
 		if (BackendToCache == BackendSource)
 		{
diff --git a/Engine/Source/Developer/Virtualization/Private/VirtualizationManager.h b/Engine/Source/Developer/Virtualization/Private/VirtualizationManager.h
--- a/Engine/Source/Developer/Virtualization/Private/VirtualizationManager.h
+++ b/Engine/Source/Developer/Virtualization/Private/VirtualizationManager.h
@@ -172,7 +172,7 @@ private:
 	bool ShouldDebugFailPulling();
 
 	void MountBackends(const FConfigFile& ConfigFile);
-	void ParseHierarchy(const FConfigFile& ConfigFile, const TCHAR* GraphName, const TCHAR* HierarchyKey, const FRegistedFactories& FactoryLookupTable, FBackendArray& PushArray);
+	void ParseHierarchy(const FConfigFile& ConfigFile, const TCHAR* GraphName, const TCHAR* HierarchyKey, const TCHAR* LegacyHierarchyKey, const FRegistedFactories& FactoryLookupTable, FBackendArray& PushArray);
 	bool CreateBackend(const FConfigFile& ConfigFile, const TCHAR* GraphName, const FString& ConfigEntryName, const FRegistedFactories& FactoryLookupTable, FBackendArray& PushArray);
 
 	void AddBackend(TUniquePtr<IVirtualizationBackend> Backend, FBackendArray& PushArray);
@@ -259,7 +259,7 @@ private:
 	TArray<TUniquePtr<IVirtualizationBackend>> AllBackends;
 
 	/** Backends used for caching operations (must support push operations). */
-	FBackendArray LocalCachableBackends; 
+	FBackendArray CacheStorageBackends; 
 
 	/** Backends used for persistent storage operations (must support push operations). */
 	FBackendArray PersistentStorageBackends; 
diff --git a/Engine/Source/Runtime/Core/Public/Virtualization/VirtualizationSystem.h b/Engine/Source/Runtime/Core/Public/Virtualization/VirtualizationSystem.h
--- a/Engine/Source/Runtime/Core/Public/Virtualization/VirtualizationSystem.h
+++ b/Engine/Source/Runtime/Core/Public/Virtualization/VirtualizationSystem.h
@@ -48,8 +48,10 @@ struct FPayloadActivityInfo
 /** Describes the type of storage to use for a given action */
 enum class EStorageType : int8
 {
+	/** Deprecated value, replaced by EStorageType::Cache */
+	Local UE_DEPRECATED(5.1, "Use EStorageType::Cache instead")  = -1,
 	/** Store in the local cache backends, this can be called from any thread */
-	Local = 0,
+	Cache = 0,
 	/** Store in the persistent backends, this can only be called from the game thread due to limitations with ISourceControlModule. */
 	Persistent
 };
diff --git a/Engine/Source/Runtime/CoreUObject/Private/Serialization/EditorBulkData.cpp b/Engine/Source/Runtime/CoreUObject/Private/Serialization/EditorBulkData.cpp
--- a/Engine/Source/Runtime/CoreUObject/Private/Serialization/EditorBulkData.cpp
+++ b/Engine/Source/Runtime/CoreUObject/Private/Serialization/EditorBulkData.cpp
@@ -1437,7 +1437,7 @@ void FEditorBulkData::PushData(const FPackagePath& InPackagePath)
 		RecompressForSerialization(PayloadToPush, Flags);
 
 		// TODO: We could make the storage type a config option?
-		if (VirtualizationSystem.PushData(PayloadContentId, PayloadToPush, UE::Virtualization::EStorageType::Local, InPackagePath.GetPackageName()))
+		if (VirtualizationSystem.PushData(PayloadContentId, PayloadToPush, UE::Virtualization::EStorageType::Cache, InPackagePath.GetPackageName()))
 		{
 			EnumAddFlags(Flags, EFlags::IsVirtualized);
 			EnumRemoveFlags(Flags, EFlags::ReferencesLegacyFile | EFlags::ReferencesWorkspaceDomain | EFlags::LegacyFileIsCompressed);
