The macro for creating backend factories (UE_REGISTER_VIRTUALIZATION_BACKEND_FACTORY) now adds the prefix F to the generated factory class.

#rnx

diff --git a/Engine/Source/Developer/Virtualization/Public/IVirtualizationBackend.h b/Engine/Source/Developer/Virtualization/Public/IVirtualizationBackend.h
--- a/Engine/Source/Developer/Virtualization/Public/IVirtualizationBackend.h
+++ b/Engine/Source/Developer/Virtualization/Public/IVirtualizationBackend.h
@@ -258,11 +258,11 @@ ENUM_CLASS_FLAGS(IVirtualizationBackend::EOperations);
  * @param The name used in config ini files to reference this backend type.
  */
 #define UE_REGISTER_VIRTUALIZATION_BACKEND_FACTORY(BackendClass, ConfigName) \
-	class BackendClass##Factory : public IVirtualizationBackendFactory \
+	class F##BackendClass##Factory : public IVirtualizationBackendFactory \
 	{ \
 	public: \
-		BackendClass##Factory() { IModularFeatures::Get().RegisterModularFeature(FName("VirtualizationBackendFactory"), this); }\
-		virtual ~BackendClass##Factory() { IModularFeatures::Get().UnregisterModularFeature(FName("VirtualizationBackendFactory"), this); } \
+		F##BackendClass##Factory() { IModularFeatures::Get().RegisterModularFeature(FName("VirtualizationBackendFactory"), this); }\
+		virtual ~F##BackendClass##Factory() { IModularFeatures::Get().UnregisterModularFeature(FName("VirtualizationBackendFactory"), this); } \
 	private: \
 		virtual TUniquePtr<IVirtualizationBackend> CreateInstance(FStringView ProjectName, FStringView ConfigName) override \
 		{ \
@@ -270,6 +270,6 @@ ENUM_CLASS_FLAGS(IVirtualizationBackend::EOperations);
 		} \
 		virtual FName GetName() override { return FName(#ConfigName); } \
 	}; \
-	static BackendClass##Factory BackendClass##Factory##Instance;
+	static F##BackendClass##Factory BackendClass##Factory##Instance;
 
 } // namespace UE::Virtualization
