Add a way to change the name used to identify a VA backend in the config file without breaking existing set ups.

#rnx

- Use the new macro UE_REGISTER_VIRTUALIZATION_BACKEND_FACTORY_LEGACY when trying to rename how backends are identified.

diff --git a/Engine/Source/Developer/Virtualization/Public/IVirtualizationBackend.h b/Engine/Source/Developer/Virtualization/Public/IVirtualizationBackend.h
--- a/Engine/Source/Developer/Virtualization/Public/IVirtualizationBackend.h
+++ b/Engine/Source/Developer/Virtualization/Public/IVirtualizationBackend.h
@@ -272,4 +272,36 @@ ENUM_CLASS_FLAGS(IVirtualizationBackend::EOperations);
 	}; \
 	static F##BackendClass##Factory BackendClass##Factory##Instance;
 
+#define UE_REGISTER_VIRTUALIZATION_BACKEND_FACTORY_LEGACY_IMPL(FactoryName, BackendClass, LegacyConfigName, ConfigName) \
+	class F##FactoryName : public IVirtualizationBackendFactory \
+	{ \
+	public: \
+		F##FactoryName(TCHAR* InLegacyConfigName, TCHAR* InNewConfigName) \
+			: StoredLegacyConfigName(InLegacyConfigName) \
+			, StoredNewConfigName(InNewConfigName) \
+			{ IModularFeatures::Get().RegisterModularFeature(FName("VirtualizationBackendFactory"), this); }\
+		virtual ~F##FactoryName() { IModularFeatures::Get().UnregisterModularFeature(FName("VirtualizationBackendFactory"), this); } \
+	private: \
+		virtual TUniquePtr<IVirtualizationBackend> CreateInstance(FStringView ProjectName, FStringView ConfigName) override \
+		{ \
+			UE_LOG(LogVirtualization, Warning, TEXT("Creating a backend via the legacy config name '%s' use '%s' instead"), *StoredLegacyConfigName, *StoredNewConfigName); \
+			return MakeUnique<BackendClass>(ProjectName, ConfigName, WriteToString<256>(#ConfigName, TEXT(" - "), ConfigName).ToString()); \
+		} \
+		virtual FName GetName() override { return FName(#LegacyConfigName); } \
+		FString StoredLegacyConfigName; \
+		FString StoredNewConfigName;\
+	}; \
+	static F##FactoryName FactoryName##Instance(TEXT(#LegacyConfigName), TEXT(#ConfigName));
+
+/** 
+ * This macro can be used to change the config name used to create backends while allowing older config file entries to continue working.
+ * If this factory is used to instantiate a backend then we will log a warning to the user so that they can update their config file.
+ * 
+ * @param BackendClass		The name of the class derived from 'IVirtualizationBackend' that the factory should create.
+ * @param LegacyConfigName	The old name that 'ConfigName' is now replacing.
+ * @param ConfigName		The name used in config ini files to reference this backend type.
+ */
+#define UE_REGISTER_VIRTUALIZATION_BACKEND_FACTORY_LEGACY(BackendClass, LegacyConfigName, ConfigName) \
+	UE_REGISTER_VIRTUALIZATION_BACKEND_FACTORY_LEGACY_IMPL(BackendClass##LegacyConfigName##To##ConfigName##Factory, BackendClass, LegacyConfigName, ConfigName)
+
 } // namespace UE::Virtualization
