Fix UE_REGISTER_VIRTUALIZATION_BACKEND_FACTORY_LEGACY_IMPL to work correctly on none MSVC platforms

#rnx

- The constructor parameters should've been const, this doesn't show when compiling with MSVC as we have that warning currently disabled.

diff --git a/Engine/Source/Developer/Virtualization/Public/IVirtualizationBackend.h b/Engine/Source/Developer/Virtualization/Public/IVirtualizationBackend.h
--- a/Engine/Source/Developer/Virtualization/Public/IVirtualizationBackend.h
+++ b/Engine/Source/Developer/Virtualization/Public/IVirtualizationBackend.h
@@ -276,7 +276,7 @@ ENUM_CLASS_FLAGS(IVirtualizationBackend::EOperations);
 	class F##FactoryName : public IVirtualizationBackendFactory \
 	{ \
 	public: \
-		F##FactoryName(TCHAR* InLegacyConfigName, TCHAR* InNewConfigName) \
+		F##FactoryName(const TCHAR* InLegacyConfigName, const TCHAR* InNewConfigName) \
 			: StoredLegacyConfigName(InLegacyConfigName) \
 			, StoredNewConfigName(InNewConfigName) \
 			{ IModularFeatures::Get().RegisterModularFeature(FName("VirtualizationBackendFactory"), this); }\
