Change the name used to identify the VA source control backend in the config file from 'SourceControl' to 'P4SourceControl' to make it clearer which source control provider is actually supported.

#rnx

- Note that the older name 'SourceControl' will continue to work but will log a warning if used to notify the user to update their config files.

diff --git a/Engine/Source/Developer/Virtualization/Private/VirtualizationSourceControlBackend.cpp b/Engine/Source/Developer/Virtualization/Private/VirtualizationSourceControlBackend.cpp
--- a/Engine/Source/Developer/Virtualization/Private/VirtualizationSourceControlBackend.cpp
+++ b/Engine/Source/Developer/Virtualization/Private/VirtualizationSourceControlBackend.cpp
@@ -897,7 +897,8 @@ void FSourceControlBackend::OnConnectionError()
 	}
 }
 
-UE_REGISTER_VIRTUALIZATION_BACKEND_FACTORY(FSourceControlBackend, SourceControl);
+UE_REGISTER_VIRTUALIZATION_BACKEND_FACTORY(FSourceControlBackend, P4SourceControl);
+UE_REGISTER_VIRTUALIZATION_BACKEND_FACTORY_LEGACY(FSourceControlBackend, SourceControl, P4SourceControl);
 
 } // namespace UE::Virtualization
 
diff --git a/Engine/Source/Developer/Virtualization/Private/VirtualizationSourceControlBackend.h b/Engine/Source/Developer/Virtualization/Private/VirtualizationSourceControlBackend.h
--- a/Engine/Source/Developer/Virtualization/Private/VirtualizationSourceControlBackend.h
+++ b/Engine/Source/Developer/Virtualization/Private/VirtualizationSourceControlBackend.h
@@ -21,7 +21,7 @@ class FSemaphore;
  * backend, found in Utils::PayloadIdToPath.
  *
  * Ini file setup:
- * 'Name'=(Type=SourceControl, DepotRoot="//XXX/", UsePartitionedClient=X, SubmitFromTempDir=X)
+ * 'Name'=(Type=P4SourceControl, DepotRoot="//XXX/", UsePartitionedClient=X, SubmitFromTempDir=X)
  * Where 'Name' is the backend name in the hierarchy and 'XXX' is the path in the source control
  * depot where the payload files are being stored.
  * 
