We no longer consider failing to restore the original workspace as a fatal error when submitting virtualized payloads to perforce.,

#rnx

- When this backend was written, it shared the source control provider interface with the rest of the editor, so if we were unable to restore the workspace it would likely break the editor in interesting and hard to diagnose ways.
- Given the scope for problems and the low change of the operatation I opted to make the error fatal.
- Now that the backend owned it's own source control provider interface, the worst thing that can happen is that future virtualization attempts will also fail, but the user will not be left with a non functioning editor. With this in mind there is no longer a reason to consider this a fatal failure.

diff --git a/Engine/Source/Developer/Virtualization/Private/VirtualizationSourceControlBackend.cpp b/Engine/Source/Developer/Virtualization/Private/VirtualizationSourceControlBackend.cpp
--- a/Engine/Source/Developer/Virtualization/Private/VirtualizationSourceControlBackend.cpp
+++ b/Engine/Source/Developer/Virtualization/Private/VirtualizationSourceControlBackend.cpp
@@ -497,12 +497,7 @@ bool FSourceControlBackend::PushData(TArrayView<FPushRequest> Requests)
 		FSourceControlResultInfo SwitchToOldWorkspaceInfo;
 		if (SCCProvider->SwitchWorkspace(OriginalWorkspace, SwitchToOldWorkspaceInfo, nullptr) != ECommandResult::Succeeded)
 		{
-			// Failing to restore the old workspace could result in confusing editor issues and data loss, so for now it is fatal.
-			// The medium term plan should be to refactor the SourceControlModule so that we could use an entirely different 
-			// ISourceControlProvider so as not to affect the rest of the editor.
-			UE_LOG(LogVirtualization, Fatal, TEXT("[%s] Failed to restore the original workspace to temp workspace '%s' continuing would risk editor instability and potential data loss"),
-					*GetDebugName(),
-					*OriginalWorkspace);
+			UE_LOG(LogVirtualization, Error, TEXT("[%s] Failed to restore the original workspace '%s'"), *GetDebugName(), *OriginalWorkspace);
 		}
 	};
 
