Virtualizing payloads when submitting via the editor will now correctly push the payloads to cache storage backends as well as persistent backends

#rnx

### Problem
- The virtualization process first checks to see if any of the local payloads are already stored in persistent storage and only try to submit those which aren't, which would mean if we are submitting a package with a local payload that is already in the persistent storage system but not in cached storage we have no way to cache it.
- This logic also introduced problems where filtered out payloads were being virtualized which was fixed with the addition of ENABLE_FILTERING_HACK which was not a robust long term solution.
- Checking if the payloads need to be pushed or not complicated the logic of the virtualization process and makes it harder to understand and make fixes too.

### Fix
- Removed code for ENABLE_FILTERING_HACK and UE_PRECHECK_PAYLOAD_STATUS from the virtualization process
- Removed some code loops from the virtualization process but we still end up with one loop finding all of the package trailers and a second loop setting up the FWorkspaceDomainPayloadProvider
-- This second loop will be removed in a future work item to avoid trying to push duplicate payloads (unless I end up doing that work at the virtualization manager level instead)
- We make one push for cached storage and another for persistent
-- Cached storage push failure will only result in a warning, failing the persistent push will error out as before
-- We need to reset the request states after the cache push. In a future work item we will likely allow a single push to specific both storage solutions, so I am leaving this code using the same set of requests for both pushes.

diff --git a/Engine/Source/Developer/Virtualization/Private/PackageVirtualizationProcess.cpp b/Engine/Source/Developer/Virtualization/Private/PackageVirtualizationProcess.cpp
--- a/Engine/Source/Developer/Virtualization/Private/PackageVirtualizationProcess.cpp
+++ b/Engine/Source/Developer/Virtualization/Private/PackageVirtualizationProcess.cpp
@@ -20,13 +20,8 @@
 #include "Virtualization/VirtualizationSystem.h"
 #include "VirtualizationManager.h"
 
-
 #define LOCTEXT_NAMESPACE "Virtualization"
 
-// When enabled we will check the payloads to see if they already exist in the persistent storage
-// backends before trying to push them.
-#define UE_PRECHECK_PAYLOAD_STATUS 1
-
 namespace UE::Virtualization
 {
 
@@ -40,6 +35,9 @@ namespace UE::Virtualization
  * 
  * So far this has shown to be a rounding error compared to the actual cost of virtualization 
  * and so implementing any level of caching has been left as a future task.
+ * 
+ * TODO: Implement a MRU cache for payloads to prevent loading the same payload off disk many
+ * times for different backends if it will not cause a huge memory spike.
  */
 class FWorkspaceDomainPayloadProvider final : public IPayloadProvider
 {
@@ -156,47 +154,6 @@ private:
 	TMap<FIoHash, FPayloadData> PayloadLookupTable;
 };
 
-#if ENABLE_FILTERING_HACK
-// This filtering provider should only ever be used with FVirtualizationManager::FilterRequests
-// and so does not need to be able to provide the payload, just the payload size.
-class FFilterProvider final : public IPayloadProvider
-{
-public:
-	void RegisterPayload(const FIoHash& PayloadId, uint64 SizeOnDisk)
-	{
-		if (!PayloadId.IsZero())
-		{
-			PayloadLookupTable.Emplace(PayloadId, SizeOnDisk);
-		}
-	}
-
-private:
-	virtual FCompressedBuffer RequestPayload(const FIoHash& Identifier)
-	{
-		checkNoEntry();
-
-		return FCompressedBuffer();
-	}
-
-	virtual uint64 GetPayloadSize(const FIoHash& Identifier)
-	{
-		if (uint64* SizeOnDisk = PayloadLookupTable.Find(Identifier))
-		{
-			return *SizeOnDisk;
-		}
-		else
-		{
-			UE_LOG(LogVirtualization, Error, TEXT("FFilterProvider was unable to find a payload with the identifier '%s'"),
-				*LexToString(Identifier));
-
-			return 0;
-		}
-	}
-
-	TMap<FIoHash, uint64> PayloadLookupTable;
-};
-#endif //ENABLE_FILTERING_HACK
-
 void VirtualizePackages(const TArray<FString>& FilesToSubmit, TArray<FText>& OutErrors)
 {
 	TRACE_CPUPROFILER_EVENT_SCOPE(UE::Virtualization::VirtualizePackages);
@@ -230,6 +187,8 @@ void VirtualizePackages(const TArray<FString>& FilesToSubmit, TArray<FText>& Out
 		FPackageTrailer Trailer;
 
 		TArray<FIoHash> LocalPayloads;
+
+		/** Index where the FPushRequest for this package can be found */
 		int32 PayloadIndex = INDEX_NONE;
 
 		bool bWasTrailerUpdated = false;
@@ -240,22 +199,13 @@ void VirtualizePackages(const TArray<FString>& FilesToSubmit, TArray<FText>& Out
 	TArray<FPackageInfo> Packages;
 	Packages.Reserve(FilesToSubmit.Num());
 
-	TArray<FIoHash> AllLocalPayloads;
-	AllLocalPayloads.Reserve(FilesToSubmit.Num());
-
 	Progress.EnterProgressFrame(1.0f);
 
-#if ENABLE_FILTERING_HACK
-	FFilterProvider FilterProvider;
-	TArray<Virtualization::FPushRequest> PayloadsToFilter;
-	PayloadsToFilter.Reserve(PayloadsToFilter.Num());
-#endif //ENABLE_FILTERING_HACK
-
 	// From the list of files to submit we need to find all of the valid packages that contain
 	// local payloads that need to be virtualized.
 	int64 TotalPackagesFound = 0;
 	int64 TotalPackageTrailersFound = 0;
-	int64 TotalPayloadsToCheck = 0;
+	int64 TotalPayloadsToVirtualize = 0;
 	for (const FString& AbsoluteFilePath : FilesToSubmit)
 	{
 		FPackagePath PackagePath = FPackagePath::FromLocalPath(AbsoluteFilePath);
@@ -274,8 +224,8 @@ void VirtualizePackages(const TArray<FString>& FilesToSubmit, TArray<FText>& Out
 				ensureMsgf(Trailer.GetNumPayloads(EPayloadStorageType::Referenced) == 0, TEXT("Trying to virtualize a package that already contains payload references which the workspace file should not ever contain!"));
 				if (Trailer.GetNumPayloads(EPayloadStorageType::Referenced) > 0)
 				{
-					FText Message = FText::Format(	LOCTEXT("Virtualization_PkgHasReferences", "Cannot virtualize the package '{1}' as it has referenced payloads in the trailer"),
-													FText::FromString(PackagePath.GetDebugName()));
+					FText Message = FText::Format(LOCTEXT("Virtualization_PkgHasReferences", "Cannot virtualize the package '{1}' as it has referenced payloads in the trailer"),
+						FText::FromString(PackagePath.GetDebugName()));
 					OutErrors.Add(Message);
 					return;
 				}
@@ -287,120 +237,33 @@ void VirtualizePackages(const TArray<FString>& FilesToSubmit, TArray<FText>& Out
 				PkgInfo.LocalPayloads = PkgInfo.Trailer.GetPayloads(EPayloadFilter::CanVirtualize);
 
 				if (!PkgInfo.LocalPayloads.IsEmpty())
-				{	
-#if ENABLE_FILTERING_HACK
-					// Build up an array of push requests that match the order of AllLocalPayloads/PayloadStatuses
-					for (const FIoHash& PayloadId : PkgInfo.LocalPayloads)
-					{
-						const uint64 SizeOnDisk = PkgInfo.Trailer.FindPayloadSizeOnDisk(PayloadId);
-
-						FilterProvider.RegisterPayload(PayloadId, SizeOnDisk);
-						PayloadsToFilter.Emplace(PayloadId, FilterProvider, PkgInfo.Path.GetPackageName());
-					}
-#endif //ENABLE_FILTERING_HACK
-
-					TotalPayloadsToCheck += PkgInfo.LocalPayloads.Num();
-
-					PkgInfo.PayloadIndex = AllLocalPayloads.Num();
-					AllLocalPayloads.Append(PkgInfo.LocalPayloads);
+				{
+					TotalPayloadsToVirtualize += PkgInfo.LocalPayloads.Num();
 
-					Packages.Emplace(MoveTemp(PkgInfo));		
+					Packages.Emplace(MoveTemp(PkgInfo));
 				}
 			}
 		}
 	}
 
 	UE_LOG(LogVirtualization, Display, TEXT("Found %" INT64_FMT " package(s), %" INT64_FMT " of which had payload trailers"), TotalPackagesFound, TotalPackageTrailersFound);
-	UE_LOG(LogVirtualization, Display, TEXT("Found %" INT64_FMT " payload(s) in %d package(s) that need to be examined for virtualization"), TotalPayloadsToCheck, Packages.Num());
-
-	Progress.EnterProgressFrame(1.0f);
-
-	TArray<EPayloadStatus> PayloadStatuses;
-	if (System.QueryPayloadStatuses(AllLocalPayloads, EStorageType::Persistent, PayloadStatuses) != EQueryResult::Success)
-	{
-		FText Message = LOCTEXT("Virtualization_DoesExistFail", "Failed to find the status of the payloads in the packages being submitted");
-		OutErrors.Add(Message);
-
-		return;
-	}
-
-#if ENABLE_FILTERING_HACK
-	{
-		check(PayloadStatuses.Num() == PayloadsToFilter.Num());
-		// If VirtualizePackages is running then we know that System is a FVirtualizationManager so we can just cast.
-		// This lets us avoid adding ::FilterRequests to IVirtualizationSystem and keeps the hack contained to this
-		// module.
-		const FVirtualizationManager* Manager = (FVirtualizationManager*)&System;
-		Manager->FilterRequests(PayloadsToFilter);
-
-		// There are many ways we could stop payloads that should be filtered from being auto virtualized if they 
-		// are present in the persistent backend, but the easiest way without changing the existing code paths is
-		// to set the status to NotFound if we know it should be filtered, to make sure that the payload is sent
-		// to the push request where it will be properly rejected by filtering.
-		for (int32 Index = 0; Index < PayloadStatuses.Num(); ++Index)
-		{
-			if (PayloadsToFilter[Index].GetStatus() != FPushRequest::EStatus::Success)
-			{
-				PayloadStatuses[Index] = EPayloadStatus::NotFound;
-			}
-		}
-	}
-#endif
 
-	// Update payloads that are already in persistent storage and don't need to be pushed
-	int64 TotalPayloadsToVirtualize = 0;
-	for (FPackageInfo& PackageInfo : Packages)
-	{
-		check(PackageInfo.LocalPayloads.IsEmpty() || PackageInfo.PayloadIndex != INDEX_NONE); // If we have payloads we should have an index
-
-#if UE_PRECHECK_PAYLOAD_STATUS
-		for (int32 Index = 0; Index < PackageInfo.LocalPayloads.Num(); ++Index)
-		{
-			if (PayloadStatuses[PackageInfo.PayloadIndex + Index] == EPayloadStatus::FoundAll)
-			{
-				if (PackageInfo.Trailer.UpdatePayloadAsVirtualized(PackageInfo.LocalPayloads[Index]))
-				{
-					PackageInfo.bWasTrailerUpdated = true;
-				}
-				else
-				{
-					FText Message = FText::Format(	LOCTEXT("Virtualization_UpdateStatusFailed", "Unable to update the status for the payload '{0}' in the package '{1}'"),
-													FText::FromString(LexToString(PackageInfo.LocalPayloads[Index])),
-													FText::FromString(PackageInfo.Path.GetDebugName()));
-					OutErrors.Add(Message);
-					return;
-				}
-			}
-		}
-
-		// If we made changes we should recalculate the local payloads left
-		if (PackageInfo.bWasTrailerUpdated)
-		{
-			PackageInfo.LocalPayloads = PackageInfo.Trailer.GetPayloads(EPayloadStorageType::Local);
-		}
-#endif
-
-		PackageInfo.PayloadIndex = INDEX_NONE;
-		TotalPayloadsToVirtualize += PackageInfo.LocalPayloads.Num();
-	}
+	// TODO: Currently not all of the filtering is done as package save time, so some of the local payloads may not get virtualized.
+	// When/if we move all filtering to package save we can change this log message to state that the local payloads *will* be virtualized.
+	UE_LOG(LogVirtualization, Display, TEXT("Found %" INT64_FMT " locally stored payload(s) in %d package(s) that maybe need to be virtualized"), TotalPayloadsToVirtualize, Packages.Num());
 
-	UE_LOG(LogVirtualization, Display, TEXT("Found %" INT64_FMT " payload(s) that potentially need to be pushed to persistent virtualized storage"), TotalPayloadsToVirtualize);
-
-	// TODO Optimization: In theory we could have many packages sharing the same payload and we only need to push once
 	Progress.EnterProgressFrame(1.0f);
 
-	// Build up the info in the payload provider and the final array of payload push requests
+	// TODO Optimization: We might want to check for duplicate payloads and remove them at this point
 
+	// Build up the info in the payload provider and the final array of payload push requests
 	FWorkspaceDomainPayloadProvider PayloadProvider;
 	TArray<Virtualization::FPushRequest> PayloadsToSubmit;
 	PayloadsToSubmit.Reserve(TotalPayloadsToVirtualize);
 
 	for (FPackageInfo& PackageInfo : Packages)
 	{
-		if (PackageInfo.LocalPayloads.IsEmpty())
-		{
-			continue;
-		}
+		check(!PackageInfo.LocalPayloads.IsEmpty());
 
 		PackageInfo.PayloadIndex = PayloadsToSubmit.Num();
 
@@ -413,21 +276,48 @@ void VirtualizePackages(const TArray<FString>& FilesToSubmit, TArray<FText>& Out
 		}
 	}
 
-	Progress.EnterProgressFrame(1.0f);
-	// Push any remaining local payload to the persistent backends
-	if (!System.PushData(PayloadsToSubmit, EStorageType::Persistent))
+	// TODO: We should be able to do both Cache and Persistent pushes in the same call
+
+	// Push payloads to cache storage
 	{
-		FText Message = LOCTEXT("Virtualization_PushFailure", "Failed to push payloads");
-		OutErrors.Add(Message);
-		return;
-	}
+		Progress.EnterProgressFrame(1.0f);
+		
+		if (!System.PushData(PayloadsToSubmit, EStorageType::Cache))
+		{
+			// Caching is not critical to the process so we only warn if it fails
+			UE_LOG(LogVirtualization, Warning, TEXT("Failed to push to EStorageType::Cache storage"));
+		}
+
+		int64 TotalPayloadsCached = 0;
+		for (Virtualization::FPushRequest& Request : PayloadsToSubmit)
+		{
+			TotalPayloadsCached += Request.GetStatus() == FPushRequest::EStatus::Success ? 1 : 0;
 
-	int64 TotalPayloadsVirtualized = 0;
-	for (const Virtualization::FPushRequest& Request : PayloadsToSubmit)
+			// TODO: This really shouldn't be required, fix when we allow both pushes to be done in the same call
+			// Reset the status for the persistent storage push
+			Request.SetStatus(Virtualization::FPushRequest::EStatus::Failed);
+		}
+		UE_LOG(LogVirtualization, Display, TEXT("Pushed %" INT64_FMT " payload(s) to cached storage"), TotalPayloadsCached);
+	}
+	
+	// Push payloads to persistent storage
 	{
-		TotalPayloadsVirtualized += Request.GetStatus() == FPushRequest::EStatus::Success ? 1 : 0;
+		Progress.EnterProgressFrame(1.0f);
+
+		if (!System.PushData(PayloadsToSubmit, EStorageType::Persistent))
+		{
+			FText Message = LOCTEXT("Virtualization_PushFailure", "Failed to push payloads");
+			OutErrors.Add(Message);
+			return;
+		}
+
+		int64 TotalPayloadsVirtualized = 0;
+		for (const Virtualization::FPushRequest& Request : PayloadsToSubmit)
+		{
+			TotalPayloadsVirtualized += Request.GetStatus() == FPushRequest::EStatus::Success ? 1 : 0;
+		}
+		UE_LOG(LogVirtualization, Display, TEXT("Pushed %" INT64_FMT " payload(s) to EStorageType::Persistent storage"), TotalPayloadsVirtualized);
 	}
-	UE_LOG(LogVirtualization, Display, TEXT("Pushed %" INT64_FMT " payload(s) to persistent virtualized storage"), TotalPayloadsVirtualized);
 
 	// Update the package info for the submitted payloads
 	for (FPackageInfo& PackageInfo : Packages)
diff --git a/Engine/Source/Developer/Virtualization/Private/VirtualizationManager.cpp b/Engine/Source/Developer/Virtualization/Private/VirtualizationManager.cpp
--- a/Engine/Source/Developer/Virtualization/Private/VirtualizationManager.cpp
+++ b/Engine/Source/Developer/Virtualization/Private/VirtualizationManager.cpp
@@ -301,37 +301,6 @@ FVirtualizationManager::~FVirtualizationManager()
 	UE_LOG(LogVirtualization, Log, TEXT("Virtualization manager destroyed"));
 }
 
-#if ENABLE_FILTERING_HACK
-void FVirtualizationManager::FilterRequests(TArray<FPushRequest>& Requests) const
-{
-	// The same filtering code as in FVirtualizationManager::PushData, since this is a hack
-	// I do not want to make any changes to real code paths so it was safer to just duplicate
-	// the checks that we need.
-	for (FPushRequest& Request : Requests)
-	{
-		if (Request.GetIdentifier().IsZero() || Request.GetPayloadSize() == 0)
-		{
-			Request.SetStatus(FPushRequest::EStatus::Invalid);
-			continue;
-		}
-
-		if ((int64)Request.GetPayloadSize() < MinPayloadLength)
-		{
-			Request.SetStatus(FPushRequest::EStatus::BelowMinSize);
-			continue;
-		}
-
-		if (!ShouldVirtualize(Request.GetContext()))
-		{
-			Request.SetStatus(FPushRequest::EStatus::ExcludedByPackagPath);
-			continue;
-		}
-
-		Request.SetStatus(FPushRequest::EStatus::Success);
-	}
-}
-#endif //ENABLE_FILTERING_HACK
-
 bool FVirtualizationManager::Initialize(const FInitParams& InitParams)
 {
 	TRACE_CPUPROFILER_EVENT_SCOPE(FVirtualizationManager::Initialize);
diff --git a/Engine/Source/Developer/Virtualization/Private/VirtualizationManager.h b/Engine/Source/Developer/Virtualization/Private/VirtualizationManager.h
--- a/Engine/Source/Developer/Virtualization/Private/VirtualizationManager.h
+++ b/Engine/Source/Developer/Virtualization/Private/VirtualizationManager.h
@@ -12,18 +12,6 @@
 class IConsoleObject;
 class FOutputDevice;
 
-// When enabled it will be possible to query push requests to see if they are filtered without actually pushing
-// anything. This is due to a bug with the filtering system, where a payload that should be filtered will be
-// virtualized by our process if the payload is already stored in persistent storage. This hack allows us to 
-// find which payloads should be filtered based on size and package path, then prevent them from being automatically
-// switched to being virtualized. Then when they are 'pushed' the usual filtering path will reject them.
-// The plan is to move this filtering to be done when the package trailer is created, on package save, in which
-// case everything will work just like asset filtering and will be much more robust. However we need to prevent 
-// some development branches from accidentally virtualizing packages which should not be virtualized so this hack
-// is a quick way to get the problem fixed and unblock workflows while the better fix is worked on. This hack
-// should be removed before 5.1 ships.
-#define ENABLE_FILTERING_HACK 1
-
 /**
  * Configuring the backend hierarchy
  * 
@@ -118,10 +106,6 @@ public:
 	FVirtualizationManager();
 	virtual ~FVirtualizationManager();
 
-#if ENABLE_FILTERING_HACK
-	void FilterRequests(TArray<FPushRequest>& Requests) const;
-#endif //ENABLE_FILTERING_HACK
-
 private:
 	/* IVirtualizationSystem implementation */
 
diff --git a/Engine/Source/Runtime/Core/Public/Virtualization/VirtualizationSystem.h b/Engine/Source/Runtime/Core/Public/Virtualization/VirtualizationSystem.h
--- a/Engine/Source/Runtime/Core/Public/Virtualization/VirtualizationSystem.h
+++ b/Engine/Source/Runtime/Core/Public/Virtualization/VirtualizationSystem.h
@@ -202,7 +202,7 @@ struct FPushRequest
 	}
 
 	/** Return the identifer used in the request */
-	FIoHash GetIdentifier() const
+	const FIoHash& GetIdentifier() const
 	{
 		return Identifier;
 	}
