The perforce source control backend now checks to see which requested payloads need uploading before doing anything else.

#rnx

- By checking the state of the payloads first we can avoid loading and writing a payload to disk, just to find that it is already in the depot. If we are very lucky we might find that all of the payloads are already uploaded, in which case we can avoid creating the submission workspace as well.
- The lock to prevent too many concurrent p4 connections does cover a wider area of code than it did before, in theory we could relinquish the lock while creating the payload files for submission, but generally the time it takes to do that is much smaller than the source control operations themselves so there is probably not much point.
- Any request that we find is already uploaded we can just mark as a success. We then create an array of pointers to the requests that were not found and operate on that for the rest of submission.
- Note that we also split up the payload existence checks into batches with the same max size limit of submission, although checking if a payload exists should be much less taxing on the system I'd rather not risk overloading the perforce server where possible.

diff --git a/Engine/Source/Developer/Virtualization/Private/VirtualizationSourceControlBackend.cpp b/Engine/Source/Developer/Virtualization/Private/VirtualizationSourceControlBackend.cpp
--- a/Engine/Source/Developer/Virtualization/Private/VirtualizationSourceControlBackend.cpp
+++ b/Engine/Source/Developer/Virtualization/Private/VirtualizationSourceControlBackend.cpp
@@ -166,7 +166,7 @@ private:
 }
 
 /** Builds a changelist description to be used when submitting a payload to source control */
-static void CreateDescription(const FString& ProjectName, const TArray<const FPushRequest*>& FileRequests, TStringBuilder<512>& OutDescription)
+static void CreateDescription(const FString& ProjectName, const TArrayView<FPushRequest*>& FileRequests, TStringBuilder<512>& OutDescription)
 {
 	// TODO: Maybe make writing out the project name an option or allow for a codename to be set via ini file?
 	OutDescription << TEXT("Submitted for project: ");
@@ -395,6 +395,70 @@ bool FSourceControlBackend::PushData(TArrayView<FPushRequest> Requests)
 {
 	TRACE_CPUPROFILER_EVENT_SCOPE(FSourceControlBackend::PushData);
 
+	// First check to see which payloads actually need pushing, if we are lucky we might
+	// not need to even bother with workspace manipulation if all payloads already exist.
+
+	FSemaphoreScopeLock _(ConcurrentConnectionLimit.Get());
+
+	TArray<FPushRequest*> RequestsToPush;
+	RequestsToPush.Reserve(Requests.Num());
+
+	{
+		TRACE_CPUPROFILER_EVENT_SCOPE(FSourceControlBackend::PushData::CheckIfPayloadsAlreadyExist);
+
+		UE_LOG(LogVirtualization, Log, TEXT("[%s] Checking %d payload(s) to see if they need uploading to source control..."), *GetDebugName(), Requests.Num());
+		const int32 NumBatches = FMath::DivideAndRoundUp<int32>(Requests.Num(), MaxBatchCount);
+
+		TArray<FIoHash> PayloadIdentifiers;
+		PayloadIdentifiers.Reserve(MaxBatchCount);
+
+		TArray<bool> PayloadResults;
+		PayloadResults.Reserve(MaxBatchCount);
+
+		for (int32 BatchIndex = 0; BatchIndex < NumBatches; ++BatchIndex)
+		{
+			PayloadIdentifiers.Reset(0);
+			PayloadResults.Reset(0);
+
+			const int32 BatchStart = BatchIndex * MaxBatchCount;
+			const int32 BatchEnd = FMath::Min((BatchIndex + 1) * MaxBatchCount, Requests.Num());
+
+			for (int32 ReqIndex = BatchStart; ReqIndex < BatchEnd; ++ReqIndex)
+			{
+				PayloadIdentifiers.Add(Requests[ReqIndex].GetIdentifier());
+			}
+
+			if (!DoPayloadsExist(PayloadIdentifiers, PayloadResults))
+			{
+				UE_LOG(LogVirtualization, Error, TEXT("[%s] Failed to find the current file state for payloads"), *GetDebugName());
+				return false;
+			}
+
+			check(PayloadIdentifiers.Num() == PayloadResults.Num());
+
+			for(int32 Index = 0; Index < PayloadResults.Num(); ++Index)
+			{
+				FPushRequest& Request = Requests[BatchStart + Index];
+
+				if (PayloadResults[Index])
+				{
+					Request.SetStatus(FPushRequest::EStatus::Success);
+				}
+				else
+				{
+					RequestsToPush.Add(&Request);
+				}
+			}
+		}
+	}
+
+	UE_LOG(LogVirtualization, Log, TEXT("[%s] Found %d payload(s) that require submission to source control"), *GetDebugName(), RequestsToPush.Num());
+
+	if (RequestsToPush.IsEmpty())
+	{
+		return true;
+	}
+
 	// TODO: Consider creating one workspace and one temp dir per session rather than per push.
 	// Although this would require more checking on start up to check for lingering workspaces
 	// and directories in case of editor crashes.
@@ -407,7 +471,7 @@ bool FSourceControlBackend::PushData(TArrayView<FPushRequest> Requests)
 
 	const FGuid SessionGuid = FGuid::NewGuid();
 	
-	UE_LOG(LogVirtualization, Log, TEXT("[%s] Started payload submission session '%s' for '%d' payload(s)"), *GetDebugName(), *LexToString(SessionGuid), Requests.Num());
+	UE_LOG(LogVirtualization, Log, TEXT("[%s] Started payload submission session '%s' for '%d' payload(s)"), *GetDebugName(), *LexToString(SessionGuid), RequestsToPush.Num());
 
 	TStringBuilder<260> SessionDirectory;
 	FPathViews::Append(SessionDirectory, SubmissionRootDir, SessionGuid);
@@ -426,8 +490,6 @@ bool FSourceControlBackend::PushData(TArrayView<FPushRequest> Requests)
 		IFileManager::Get().DeleteDirectory(SessionDirectory.ToString(), false, true);
 	};
 
-	FSemaphoreScopeLock _(ConcurrentConnectionLimit.Get());
-
 	TStringBuilder<64> WorkspaceName;
 	WorkspaceName << TEXT("VASubmission-") << SessionGuid;
 
@@ -501,18 +563,18 @@ bool FSourceControlBackend::PushData(TArrayView<FPushRequest> Requests)
 		}
 	};
 
-	const int32 NumBatches = FMath::DivideAndRoundUp<int32>(Requests.Num(), MaxBatchCount);
+	const int32 NumBatches = FMath::DivideAndRoundUp<int32>(RequestsToPush.Num(), MaxBatchCount);
 
-	UE_LOG(LogVirtualization, Log, TEXT("[%s] Splitting the push into '%d' batches"), *GetDebugName(), NumBatches);
+	UE_LOG(LogVirtualization, Log, TEXT("[%s] Splitting the submission into '%d' batches"), *GetDebugName(), NumBatches);
 
 	for (int32 BatchIndex = 0; BatchIndex < NumBatches; ++BatchIndex)
 	{
 		UE_LOG(LogVirtualization, Log, TEXT("[%s] Processing batch %d/%d..."), *GetDebugName(), BatchIndex + 1, NumBatches);
 
 		const int32 BatchStart = BatchIndex * MaxBatchCount;
-		const int32 BatchEnd = FMath::Min((BatchIndex + 1) * MaxBatchCount, Requests.Num());
+		const int32 BatchEnd = FMath::Min((BatchIndex + 1) * MaxBatchCount, RequestsToPush.Num());
 
-		TArrayView<FPushRequest> RequestBatch(&Requests[BatchStart], BatchEnd - BatchStart);
+		TArrayView<FPushRequest*> RequestBatch(&RequestsToPush[BatchStart], BatchEnd - BatchStart);
 
 		TArray<FString> FilesToSubmit;
 		FilesToSubmit.Reserve(RequestBatch.Num());
@@ -522,22 +584,24 @@ bool FSourceControlBackend::PushData(TArrayView<FPushRequest> Requests)
 		{
 			TRACE_CPUPROFILER_EVENT_SCOPE(FSourceControlBackend::PushData::CreateFiles);
 
-			for (const FPushRequest& Request : RequestBatch)
+			for (const FPushRequest* Request : RequestBatch)
 			{
+				check(Request != nullptr);
+
 				TStringBuilder<52> LocalPayloadPath;
-				Utils::PayloadIdToPath(Request.GetIdentifier(), LocalPayloadPath);
+				Utils::PayloadIdToPath(Request->GetIdentifier(), LocalPayloadPath);
 
 				TStringBuilder<260> PayloadFilePath;
 				FPathViews::Append(PayloadFilePath, SessionDirectory, LocalPayloadPath);
 
 				UE_LOG(LogVirtualization, Verbose, TEXT("[%s] Writing payload to '%s' for submission"), *GetDebugName(), PayloadFilePath.ToString());
 
-				FCompressedBuffer Payload = Request.GetPayload();
+				FCompressedBuffer Payload = Request->GetPayload();
 				if (Payload.IsNull())
 				{
 					UE_LOG(LogVirtualization, Error, TEXT("[%s] Failed to acquire payload '%s' contents to '%s' for writing"),
 						*GetDebugName(),
-						*LexToString(Request.GetIdentifier()),
+						*LexToString(Request->GetIdentifier()),
 						PayloadFilePath.ToString());
 
 					return false;
@@ -551,7 +615,7 @@ bool FSourceControlBackend::PushData(TArrayView<FPushRequest> Requests)
 
 					UE_LOG(LogVirtualization, Error, TEXT("[%s] Failed to write payload '%s' contents to '%s' due to system error: %s"),
 						*GetDebugName(),
-						*LexToString(Request.GetIdentifier()),
+						*LexToString(Request->GetIdentifier()),
 						PayloadFilePath.ToString(),
 						SystemErrorMsg.ToString());
 
@@ -567,7 +631,7 @@ bool FSourceControlBackend::PushData(TArrayView<FPushRequest> Requests)
 
 					UE_LOG(LogVirtualization, Error, TEXT("[%s] Failed to write payload '%s' contents to '%s' due to system error: %s"),
 						*GetDebugName(),
-						*LexToString(Request.GetIdentifier()),
+						*LexToString(Request->GetIdentifier()),
 						*PayloadFilePath,
 						SystemErrorMsg.
 						ToString());
@@ -581,55 +645,10 @@ bool FSourceControlBackend::PushData(TArrayView<FPushRequest> Requests)
 
 		check(RequestBatch.Num() == FilesToSubmit.Num());
 
-		TArray<FSourceControlStateRef> FileStates;
-		{
-			TRACE_CPUPROFILER_EVENT_SCOPE(FSourceControlBackend::PushData::GetFileStates);
-			if (GetDepotPathStates(*SCCProvider, FilesToSubmit, FileStates) != ECommandResult::Succeeded)
-			{
-				UE_LOG(LogVirtualization, Error, TEXT("[%s] Failed to find the current file state for payloads"), *GetDebugName());
-				return false;
-			}
-		}
-		check(RequestBatch.Num() == FileStates.Num());
-
-		TArray<FString> FilesToAdd;
-		FilesToAdd.Reserve(FilesToSubmit.Num());
-
-		TArray<const FPushRequest*> FileRequests;
-		FileRequests.Reserve(FilesToSubmit.Num());
-
-		for (int32 FileIndex = 0; FileIndex < FilesToSubmit.Num(); ++FileIndex)
-		{
-			if (FileStates[FileIndex]->IsSourceControlled())
-			{
-				// TODO: Maybe check if the data is the same (could be different if the compression algorithm has changed)
-				// TODO: Should we respect if the file is deleted as technically we can still get access to it?
-				RequestBatch[FileIndex].SetStatus(FPushRequest::EStatus::Success);
-			}
-			else if (FileStates[FileIndex]->CanAdd())
-			{
-				FilesToAdd.Add(FilesToSubmit[FileIndex]);
-				FileRequests.Add(&RequestBatch[FileIndex]);
-			}
-			else
-			{
-				UE_LOG(LogVirtualization, Error, TEXT("[%s] The the payload file '%s' is not in source control but also cannot be marked for Add"), *GetDebugName(), *FilesToSubmit[FileIndex]);
-				return false;
-			}
-		}
-
-		check(FileRequests.Num() == FilesToAdd.Num());
-
-		if (FilesToAdd.IsEmpty())
-		{
-			// If we have no files to add then we should skip to the next batch
-			continue;
-		}
-
 		{
 			TRACE_CPUPROFILER_EVENT_SCOPE(FSourceControlBackend::PushData::AddFiles);
 
-			if (SCCProvider->Execute(ISourceControlOperation::Create<FMarkForAdd>(), FilesToAdd) != ECommandResult::Succeeded)
+			if (SCCProvider->Execute(ISourceControlOperation::Create<FMarkForAdd>(), FilesToSubmit) != ECommandResult::Succeeded)
 			{
 				UE_LOG(LogVirtualization, Error, TEXT("[%s] Failed to mark the payload file for Add in source control"), *GetDebugName());
 				return false;
@@ -643,11 +662,11 @@ bool FSourceControlBackend::PushData(TArrayView<FPushRequest> Requests)
 			TSharedRef<FCheckIn, ESPMode::ThreadSafe> CheckInOperation = ISourceControlOperation::Create<FCheckIn>();
 
 			TStringBuilder<512> Description;
-			CreateDescription(ProjectName, FileRequests, Description);
+			CreateDescription(ProjectName, RequestBatch, Description);
 
 			CheckInOperation->SetDescription(FText::FromString(Description.ToString()));
 
-			if (SCCProvider->Execute(CheckInOperation, FilesToAdd) != ECommandResult::Succeeded)
+			if (SCCProvider->Execute(CheckInOperation, FilesToSubmit) != ECommandResult::Succeeded)
 			{
 				UE_LOG(LogVirtualization, Error, TEXT("[%s] Failed to submit the payload file(s) to source control"), *GetDebugName());
 				return false;
@@ -655,9 +674,9 @@ bool FSourceControlBackend::PushData(TArrayView<FPushRequest> Requests)
 		}
 
 		// TODO: We really should be setting a more fine grain status for each request, or not bother with the status at all
-		for (FPushRequest& Request : RequestBatch)
+		for (FPushRequest* Request : RequestBatch)
 		{
-			Request.SetStatus(FPushRequest::EStatus::Success);
+			Request->SetStatus(FPushRequest::EStatus::Success);
 		}
 
 		// Try to clean up the files from this batch
