Fix errors given when trying to virtualize large numbers of assets, some of which contain duplicate payloads

#rnx

- The problem with the old logic was that we checked all of the payloads initially to see which needed to be submitted but if enough payloads were requested we would split the submission into a number of batches. If a payload was in more than once batch then it would be submitted in the first batch and then fail the MarkForAdd command in the subsequent batches. When we detected this error we would fail the virtualization process.
- In the long term we should probably eliminate duplicate requests at the highest level (FVirtualizationManager) so no backend has to worry about it, but we need this working now so it was easiest to add duplicate checks to the source control backend
-- Note that if another process submits a payload between the status check and the mark for add calls then we will also get errors, but that is true for a lot of our source control code.
- Now when pushing payloads we first iterate over the requests and create a unique list of requests to submit. This eliminates the possibility of duplicate payloads being submitted and failing the MarkForAdd command.
- We now record the push request state for each payload in a map and then only apply it to the original requests if the push succeeds.
-- We still only really have failed/success statuses and nothing recording "already in the backend" which is scheduled to be fixed elsewhere in the backlog.

diff --git a/Engine/Source/Developer/Virtualization/Private/VirtualizationSourceControlBackend.cpp b/Engine/Source/Developer/Virtualization/Private/VirtualizationSourceControlBackend.cpp
--- a/Engine/Source/Developer/Virtualization/Private/VirtualizationSourceControlBackend.cpp
+++ b/Engine/Source/Developer/Virtualization/Private/VirtualizationSourceControlBackend.cpp
@@ -166,7 +166,7 @@ private:
 }
 
 /** Builds a changelist description to be used when submitting a payload to source control */
-static void CreateDescription(const FString& ProjectName, const TArrayView<FPushRequest*>& FileRequests, TStringBuilder<512>& OutDescription)
+static void CreateDescription(const FString& ProjectName, TArrayView<const FPushRequest*>& FileRequests, TStringBuilder<512>& OutDescription)
 {
 	// TODO: Maybe make writing out the project name an option or allow for a codename to be set via ini file?
 	OutDescription << TEXT("Submitted for project: ");
@@ -400,14 +400,39 @@ bool FSourceControlBackend::PushData(TArrayView<FPushRequest> Requests)
 
 	FSemaphoreScopeLock _(ConcurrentConnectionLimit.Get());
 
-	TArray<FPushRequest*> RequestsToPush;
-	RequestsToPush.Reserve(Requests.Num());
+	// TODO: At some point FVirtualizationManager will do this so we know there cannot be duplicate
+	// payload requests, at which point we can rework this code.
+	// We create a new array of requests that we know are unique along with a map of the request 
+	// results that we can use to set the original requests statuses at the end.
+
+	TMap<FIoHash, FPushRequest::EStatus> PayloadStatusMap;
+	PayloadStatusMap.Reserve(Requests.Num());
+
+	TArray<const FPushRequest*> UniqueRequests;
+	UniqueRequests.Reserve(Requests.Num());
+
+	{
+		TRACE_CPUPROFILER_EVENT_SCOPE(FSourceControlBackend::PushData::FindUniquePayloads);
+
+		for (const FPushRequest& Request : Requests)
+		{
+			if (!PayloadStatusMap.Contains(Request.GetIdentifier()))
+			{
+				PayloadStatusMap.Add(Request.GetIdentifier(), FPushRequest::EStatus::Failed);
+				UniqueRequests.Add(&Request);
+			}
+		}
+
+		UE_LOG(LogVirtualization, Log, TEXT("[%s] Found %d unique payload(s)"), *GetDebugName(), UniqueRequests.Num());
+	}
+
+	TArray<const FPushRequest*> RequestsToPush;
+	RequestsToPush.Reserve(UniqueRequests.Num());
 
 	{
 		TRACE_CPUPROFILER_EVENT_SCOPE(FSourceControlBackend::PushData::CheckIfPayloadsAlreadyExist);
 
-		UE_LOG(LogVirtualization, Log, TEXT("[%s] Checking %d payload(s) to see if they need uploading to source control..."), *GetDebugName(), Requests.Num());
-		const int32 NumBatches = FMath::DivideAndRoundUp<int32>(Requests.Num(), MaxBatchCount);
+		const int32 NumBatches = FMath::DivideAndRoundUp<int32>(UniqueRequests.Num(), MaxBatchCount);
 
 		TArray<FIoHash> PayloadIdentifiers;
 		PayloadIdentifiers.Reserve(MaxBatchCount);
@@ -421,11 +446,11 @@ bool FSourceControlBackend::PushData(TArrayView<FPushRequest> Requests)
 			PayloadResults.Reset(0);
 
 			const int32 BatchStart = BatchIndex * MaxBatchCount;
-			const int32 BatchEnd = FMath::Min((BatchIndex + 1) * MaxBatchCount, Requests.Num());
+			const int32 BatchEnd = FMath::Min((BatchIndex + 1) * MaxBatchCount, UniqueRequests.Num());
 
 			for (int32 ReqIndex = BatchStart; ReqIndex < BatchEnd; ++ReqIndex)
 			{
-				PayloadIdentifiers.Add(Requests[ReqIndex].GetIdentifier());
+				PayloadIdentifiers.Add(UniqueRequests[ReqIndex]->GetIdentifier());
 			}
 
 			if (!DoPayloadsExist(PayloadIdentifiers, PayloadResults))
@@ -436,26 +461,31 @@ bool FSourceControlBackend::PushData(TArrayView<FPushRequest> Requests)
 
 			check(PayloadIdentifiers.Num() == PayloadResults.Num());
 
-			for(int32 Index = 0; Index < PayloadResults.Num(); ++Index)
+			for (int32 Index = 0; Index < PayloadResults.Num(); ++Index)
 			{
-				FPushRequest& Request = Requests[BatchStart + Index];
+				const FPushRequest* Request = UniqueRequests[BatchStart + Index];
 
 				if (PayloadResults[Index])
 				{
-					Request.SetStatus(FPushRequest::EStatus::Success);
+					PayloadStatusMap[Request->GetIdentifier()] = FPushRequest::EStatus::Success;
 				}
 				else
 				{
-					RequestsToPush.Add(&Request);
+					RequestsToPush.Add(Request);
 				}
 			}
 		}
 	}
 
-	UE_LOG(LogVirtualization, Log, TEXT("[%s] Found %d payload(s) that require submission to source control"), *GetDebugName(), RequestsToPush.Num());
+	UE_LOG(LogVirtualization, Log, TEXT("[%s] Determined that %d payload(s) require submission to source control"), *GetDebugName(), RequestsToPush.Num());
 
 	if (RequestsToPush.IsEmpty())
 	{
+		for (FPushRequest& Request : Requests)
+		{
+			Request.SetStatus(PayloadStatusMap[Request.GetIdentifier()]);
+		}
+
 		return true;
 	}
 
@@ -574,7 +604,7 @@ bool FSourceControlBackend::PushData(TArrayView<FPushRequest> Requests)
 		const int32 BatchStart = BatchIndex * MaxBatchCount;
 		const int32 BatchEnd = FMath::Min((BatchIndex + 1) * MaxBatchCount, RequestsToPush.Num());
 
-		TArrayView<FPushRequest*> RequestBatch(&RequestsToPush[BatchStart], BatchEnd - BatchStart);
+		TArrayView<const FPushRequest*> RequestBatch(&RequestsToPush[BatchStart], BatchEnd - BatchStart);
 
 		TArray<FString> FilesToSubmit;
 		FilesToSubmit.Reserve(RequestBatch.Num());
@@ -674,9 +704,9 @@ bool FSourceControlBackend::PushData(TArrayView<FPushRequest> Requests)
 		}
 
 		// TODO: We really should be setting a more fine grain status for each request, or not bother with the status at all
-		for (FPushRequest* Request : RequestBatch)
+		for (const FPushRequest* Request : RequestBatch)
 		{
-			Request->SetStatus(FPushRequest::EStatus::Success);
+			PayloadStatusMap[Request->GetIdentifier()] = FPushRequest::EStatus::Success;
 		}
 
 		// Try to clean up the files from this batch
@@ -690,6 +720,12 @@ bool FSourceControlBackend::PushData(TArrayView<FPushRequest> Requests)
 		}
 	}
 
+	// Finally set all of the request statuses
+	for (FPushRequest& Request : Requests)
+	{
+		Request.SetStatus(PayloadStatusMap[Request.GetIdentifier()]);
+	}
+
 	return true;
 }
 
