The console command 'DumpPackagePayloadInfo' now works with absolute file paths referencing packages from outside the current set of mount points.

#rnx

- Sometimes it is useful to be able to dump the payload info from packages not under the current project as a quick form of debugging and given the format of the package trailer there is no reason why we should allow it.
- If we fail to turn the arg into a valid FPackagePath can now also check to see if the arg is an absolute file path via the FileManager
- If both of these options fail we then log an error as before.

diff --git a/Engine/Source/Developer/Virtualization/Private/DumpPackagePayloadInfoCommand.cpp b/Engine/Source/Developer/Virtualization/Private/DumpPackagePayloadInfoCommand.cpp
--- a/Engine/Source/Developer/Virtualization/Private/DumpPackagePayloadInfoCommand.cpp
+++ b/Engine/Source/Developer/Virtualization/Private/DumpPackagePayloadInfoCommand.cpp
@@ -5,6 +5,7 @@
 #include "Misc/PackagePath.h"
 #include "UObject/PackageTrailer.h"
 #include "Misc/PackageName.h"
+#include "HAL/FileManager.h"
 
 namespace UE
 {
@@ -46,62 +47,75 @@ void DumpPackagePayloadInfo(const TArray<FString>& Args)
 
 	for (const FString& Arg : Args)
 	{
-		FPackagePath Path;
-		FStringView ObjectPathString;
+		FString PathString;
 
-		if (FPackageName::ParseExportTextPath(Arg, nullptr /*OutClassName*/, &ObjectPathString))
+		if (FPackageName::ParseExportTextPath(Arg, nullptr /*OutClassName*/, &PathString))
 		{
-			ObjectPathString = FPackageName::ObjectPathToPackageName(ObjectPathString);
+			PathString = FPackageName::ObjectPathToPackageName(PathString);
 		}
 		else
 		{
-			ObjectPathString = Arg;
+			PathString = Arg;
 		}
 
-		if (FPackagePath::TryFromMountedName(ObjectPathString, Path))
-		{
-			FPackageTrailer Trailer;
+		FPackageTrailer Trailer;
+
+		FPackagePath Path;
+		if (FPackagePath::TryFromMountedName(PathString, Path))
+		{	
 			if (!FPackageTrailer::TryLoadFromPackage(Path, Trailer))
 			{
 				UE_LOG(LogVirtualization, Error, TEXT("Failed to find load the package trailer from: '%s'"), *Path.GetDebugName());
 				continue;
 			}
-
-			TArray<FIoHash> LocalPayloadIds = Trailer.GetPayloads(UE::EPayloadStorageType::Local);
-			TArray<FIoHash> VirtualizedPayloadIds = Trailer.GetPayloads(UE::EPayloadStorageType::Virtualized);
-			
-			UE_LOG(LogVirtualization, Display, TEXT("")); // Blank line to make the output easier to read
-			UE_LOG(LogVirtualization, Display, TEXT("Package: '%s' has %d local and %d virtualized payloads"), *Path.GetDebugName(), LocalPayloadIds.Num(), VirtualizedPayloadIds.Num());
+		}
+		else if (IFileManager::Get().FileExists(*PathString))
+		{
+			// IF we couldn't turn it into a FPackagePath it could be a path to a package not under any current mount point.
+			// So for a final attempt we will see if we can find the file on disk and load the package trailer that way.
 			
-			if (LocalPayloadIds.Num() > 0)
+			if (!FPackageTrailer::TryLoadFromFile(PathString, Trailer))
 			{
-				UE_LOG(LogVirtualization, Display, TEXT("LocalPayloads:"));
-				UE_LOG(LogVirtualization, Display, TEXT("Index | %-40s | SizeOnDisk | FilterReason"), TEXT("PayloadIdentifier"));
-				for (int32 Index = 0; Index < LocalPayloadIds.Num(); ++Index)
-				{
-					FPayloadInfo Info = Trailer.GetPayloadInfo(LocalPayloadIds[Index]);
-					UE_LOG(LogVirtualization, Display, TEXT("%02d    | %s | %-10s | %s"), 
-						Index, 
-						*LexToString(LocalPayloadIds[Index]), 
-						*BytesToString(Info.CompressedSize),
-						*LexToString(Info.FilterFlags));
-				}
-			}
+				UE_LOG(LogVirtualization, Error, TEXT("Failed to find load the package trailer from: '%s'"), *PathString);
+				continue;
+			}	
+		}
+		else
+		{
+			UE_LOG(LogVirtualization, Error, TEXT("Arg '%s' could not be converted to a valid package path"), *Arg);
+			continue;
+		}
 
-			if (VirtualizedPayloadIds.Num() > 0)
+		TArray<FIoHash> LocalPayloadIds = Trailer.GetPayloads(UE::EPayloadStorageType::Local);
+		TArray<FIoHash> VirtualizedPayloadIds = Trailer.GetPayloads(UE::EPayloadStorageType::Virtualized);
+
+		UE_LOG(LogVirtualization, Display, TEXT("")); // Blank line to make the output easier to read
+		UE_LOG(LogVirtualization, Display, TEXT("Package: '%s' has %d local and %d virtualized payloads"), *Path.GetDebugName(), LocalPayloadIds.Num(), VirtualizedPayloadIds.Num());
+
+		if (LocalPayloadIds.Num() > 0)
+		{
+			UE_LOG(LogVirtualization, Display, TEXT("LocalPayloads:"));
+			UE_LOG(LogVirtualization, Display, TEXT("Index | %-40s | SizeOnDisk | FilterReason"), TEXT("PayloadIdentifier"));
+			for (int32 Index = 0; Index < LocalPayloadIds.Num(); ++Index)
 			{
-				UE_LOG(LogVirtualization, Display, TEXT("VirtualizedPayloads:"));
-				UE_LOG(LogVirtualization, Display, TEXT("Index|\t%-40s|\tFilterReason"), TEXT("PayloadIdentifier"));
-				for (int32 Index = 0; Index < VirtualizedPayloadIds.Num(); ++Index)
-				{
-					FPayloadInfo Info = Trailer.GetPayloadInfo(VirtualizedPayloadIds[Index]);
-					UE_LOG(LogVirtualization, Display, TEXT("%02d:  |\t%s|\t%s"), Index, *LexToString(VirtualizedPayloadIds[Index]), *LexToString(Info.FilterFlags));
-				}
+				FPayloadInfo Info = Trailer.GetPayloadInfo(LocalPayloadIds[Index]);
+				UE_LOG(LogVirtualization, Display, TEXT("%02d    | %s | %-10s | %s"),
+					Index,
+					*LexToString(LocalPayloadIds[Index]),
+					*BytesToString(Info.CompressedSize),
+					*LexToString(Info.FilterFlags));
 			}
 		}
-		else
+
+		if (VirtualizedPayloadIds.Num() > 0)
 		{
-			UE_LOG(LogVirtualization, Error, TEXT("Arg '%s' could not be converted to a valid package path"), *Arg);
+			UE_LOG(LogVirtualization, Display, TEXT("VirtualizedPayloads:"));
+			UE_LOG(LogVirtualization, Display, TEXT("Index|\t%-40s|\tFilterReason"), TEXT("PayloadIdentifier"));
+			for (int32 Index = 0; Index < VirtualizedPayloadIds.Num(); ++Index)
+			{
+				FPayloadInfo Info = Trailer.GetPayloadInfo(VirtualizedPayloadIds[Index]);
+				UE_LOG(LogVirtualization, Display, TEXT("%02d:  |\t%s|\t%s"), Index, *LexToString(VirtualizedPayloadIds[Index]), *LexToString(Info.FilterFlags));
+			}
 		}
 	}
 }
