Fix some typos in error messages for the DumpPackagePayloadInfo command

#rnx

diff --git a/Engine/Source/Developer/Virtualization/Private/DumpPackagePayloadInfoCommand.cpp b/Engine/Source/Developer/Virtualization/Private/DumpPackagePayloadInfoCommand.cpp
--- a/Engine/Source/Developer/Virtualization/Private/DumpPackagePayloadInfoCommand.cpp
+++ b/Engine/Source/Developer/Virtualization/Private/DumpPackagePayloadInfoCommand.cpp
@@ -65,7 +65,7 @@ void DumpPackagePayloadInfo(const TArray<FString>& Args)
 		{	
 			if (!FPackageTrailer::TryLoadFromPackage(Path, Trailer))
 			{
-				UE_LOG(LogVirtualization, Error, TEXT("Failed to find load the package trailer from: '%s'"), *Path.GetDebugName());
+				UE_LOG(LogVirtualization, Error, TEXT("Failed to load the package trailer from: '%s'"), *Path.GetDebugName());
 				continue;
 			}
 		}
@@ -76,7 +76,7 @@ void DumpPackagePayloadInfo(const TArray<FString>& Args)
 			
 			if (!FPackageTrailer::TryLoadFromFile(PathString, Trailer))
 			{
-				UE_LOG(LogVirtualization, Error, TEXT("Failed to find load the package trailer from: '%s'"), *PathString);
+				UE_LOG(LogVirtualization, Error, TEXT("Failed to load the package trailer from: '%s'"), *PathString);
 				continue;
 			}	
 		}
