Improve the logging of VA profiling data at the end of a cook.

#rnx

- If there was no VA activity we explicitly log that there was no profile data to show
- Only log the operations that had activity, there is no point writing out a bunch of zeros for an operation to show nothing happened.

diff --git a/Engine/Source/Developer/Virtualization/Private/VirtualizationManager.cpp b/Engine/Source/Developer/Virtualization/Private/VirtualizationManager.cpp
--- a/Engine/Source/Developer/Virtualization/Private/VirtualizationManager.cpp
+++ b/Engine/Source/Developer/Virtualization/Private/VirtualizationManager.cpp
@@ -168,42 +168,43 @@ namespace Profiling
 		return *PullStats.Find(Backend.GetDebugName());
 	}
 
-	/** Returns true if we have gathered any profiling data at all */
-	bool HasProfilingData()
+	bool HasProfilingData(const TMap<FString, FCookStats::CallStats>& Stats)
 	{
-		auto HasAccumulatedData = [](const TMap<FString, FCookStats::CallStats>& Stats)->bool
+		for (const auto& Iterator : Stats)
 		{
-			for (const auto& Iterator : Stats)
+			if (Iterator.Value.GetAccumulatedValueAnyThread(FCookStats::CallStats::EHitOrMiss::Hit, FCookStats::CallStats::EStatType::Counter) > 0)
 			{
-				if (Iterator.Value.GetAccumulatedValueAnyThread(FCookStats::CallStats::EHitOrMiss::Hit, FCookStats::CallStats::EStatType::Counter) > 0)
-				{
-					return true;
-				}
+				return true;
+			}
 
-				if (Iterator.Value.GetAccumulatedValueAnyThread(FCookStats::CallStats::EHitOrMiss::Miss, FCookStats::CallStats::EStatType::Counter) > 0)
-				{
-					return true;
-				}
+			if (Iterator.Value.GetAccumulatedValueAnyThread(FCookStats::CallStats::EHitOrMiss::Miss, FCookStats::CallStats::EStatType::Counter) > 0)
+			{
+				return true;
 			}
+		}
 
-			return false;
-		};
+		return false;
+	}
 
-		return HasAccumulatedData(CacheStats) || HasAccumulatedData(PushStats) || HasAccumulatedData(PullStats);
+	/** Returns true if we have gathered any profiling data at all */
+	bool HasProfilingData()
+	{
+		return HasProfilingData(CacheStats) || HasProfilingData(PushStats) || HasProfilingData(PullStats);
 	}
 
 	void LogStats()
 	{
+		UE_LOG(LogVirtualization, Display, TEXT(""));
+		UE_LOG(LogVirtualization, Display, TEXT("Virtualization ProfileData"));
+		UE_LOG(LogVirtualization, Display, TEXT("======================================================================================="));
+
 		if (!HasProfilingData())
 		{
+			UE_LOG(LogVirtualization, Display, TEXT("Skipping profile data as there was no activity to report"));
 			return; // Early out if we have no data
 		}
 
-		UE_LOG(LogVirtualization, Display, TEXT(""));
-		UE_LOG(LogVirtualization, Display, TEXT("Virtualization ProfileData"));
-		UE_LOG(LogVirtualization, Display, TEXT("======================================================================================="));
-
-		if (CacheStats.Num() > 0)
+		if (HasProfilingData(CacheStats))
 		{
 			UE_LOG(LogVirtualization, Display, TEXT("%-40s|%17s|%12s|%14s|"), TEXT("Caching Data"), TEXT("TotalSize (MB)"), TEXT("TotalTime(s)"), TEXT("DataRate(MB/S)"));
 			UE_LOG(LogVirtualization, Display, TEXT("----------------------------------------|-----------------|------------|--------------|"));
@@ -224,7 +225,7 @@ namespace Profiling
 			UE_LOG(LogVirtualization, Display, TEXT("======================================================================================="));
 		}
 
-		if (PushStats.Num() > 0)
+		if (HasProfilingData(PushStats))
 		{
 			UE_LOG(LogVirtualization, Display, TEXT("%-40s|%17s|%12s|%14s|"), TEXT("Pushing Data"), TEXT("TotalSize (MB)"), TEXT("TotalTime(s)"), TEXT("DataRate(MB/S)"));
 			UE_LOG(LogVirtualization, Display, TEXT("----------------------------------------|-----------------|------------|--------------|"));
@@ -245,7 +246,7 @@ namespace Profiling
 			UE_LOG(LogVirtualization, Display, TEXT("======================================================================================="));
 		}
 
-		if (PullStats.Num() > 0)
+		if (HasProfilingData(PullStats))
 		{
 			UE_LOG(LogVirtualization, Display, TEXT("%-40s|%17s|%12s|%14s|"), TEXT("Pulling Data"), TEXT("TotalSize (MB)"), TEXT("TotalTime(s)"), TEXT("DataRate(MB/S)"));
 			UE_LOG(LogVirtualization, Display, TEXT("----------------------------------------|-----------------|------------|--------------|"));
