Add a number of ways for the VA backend connections to be changed to lazy initialize on first use.

#rnx

- There are some use cases that require the VA system to be initialized and configured correctly but would prefer that the backend connections only run if absolutely needed (usually when a payload is pulled or pushed for the first time), this change provides four different ways of doing this:
-- Setting [Core.VirtualizationModule]LazyInitConnections=true in the Engine ini file
-- Setting the define 'UE_VIRTUALIZATION_CONNECTION_LAZY_INIT' to 1 in a programs .target.cs
-- Running with the commandline option -VA-LazyInitConnections
-- Setting the cvar 'VA.LazyInitConnections' to 1 (only works if it is set before the VA system is initialized, changing it mid editor via the console does nothing)
--- Note that after the config file, each setting there only opts into lazy initializing the connections, setting the cvar to 0 for example will not prevent the cmdline from opting in etc.

- In the future we will allow the connection code to run async, so the latency can be hidden behind the editor loading, but for the current use case we are taking the minimal approach.
-- This means we only support the backend being in 3 states. No connection has been made yet, the connection is broken and the connection is working.
-- To keep things simple we only record if we have attempted to connect the backends or not. We don't check individual backends nor do we try to reconnect failed ones etc. This is all scheduled for a future work item.
- If the connections are not initialized when the VA system is, we wait until the first time someone calls one of the virtualization methods that will actually use a connection: Push/Pull/Query
-- We try connecting all of the backends at once, even if they won't be used in the call to keep things simple.
- Only the source control backend makes use of the connection system. The horde storage (http) backend could take advantage too, but it is currently unused and most likely going to just be deleted so there seemed little point updating it.
- If we try to run an operation on an unconnected backend we only log to verbose. This is to maintain existing behaviour where a failed backend would not be mounted at all. This logging will likely be revisited in a future work item.

diff --git a/Engine/Config/BaseEngine.ini b/Engine/Config/BaseEngine.ini
--- a/Engine/Config/BaseEngine.ini
+++ b/Engine/Config/BaseEngine.ini
@@ -1254,6 +1254,7 @@ FilterMode=OptOut
 FilterEngineContent=True
 FilterEnginePluginContent=True
 FilterMapContent=True
+LazyInitConnections=False
 ; The following assets should not virtualize until it is confirmed that they do not crash when virtualized payloads are missing
 +DisabledAsset=StaticMesh
 +DisabledAsset=SoundWave
diff --git a/Engine/Source/Developer/Virtualization/Private/VirtualizationDDCBackend.cpp b/Engine/Source/Developer/Virtualization/Private/VirtualizationDDCBackend.cpp
--- a/Engine/Source/Developer/Virtualization/Private/VirtualizationDDCBackend.cpp
+++ b/Engine/Source/Developer/Virtualization/Private/VirtualizationDDCBackend.cpp
@@ -69,6 +69,11 @@ bool FDDCBackend::Initialize(const FString& ConfigEntry)
 	return true;	
 }
 
+IVirtualizationBackend::EConnectionStatus FDDCBackend::OnConnect()
+{
+	return IVirtualizationBackend::EConnectionStatus::Connected;
+}
+
 EPushResult FDDCBackend::PushData(const FIoHash& Id, const FCompressedBuffer& Payload, const FString& PackageContext)
 {
 	TRACE_CPUPROFILER_EVENT_SCOPE(FDDCBackend::PushData);
diff --git a/Engine/Source/Developer/Virtualization/Private/VirtualizationDDCBackend.h b/Engine/Source/Developer/Virtualization/Private/VirtualizationDDCBackend.h
--- a/Engine/Source/Developer/Virtualization/Private/VirtualizationDDCBackend.h
+++ b/Engine/Source/Developer/Virtualization/Private/VirtualizationDDCBackend.h
@@ -36,6 +36,8 @@ private:
 
 	virtual bool Initialize(const FString& ConfigEntry) override;
 
+	virtual EConnectionStatus OnConnect() override;
+
 	virtual EPushResult PushData(const FIoHash& Id, const FCompressedBuffer& Payload, const FString& PackageContext) override;
 	virtual bool PushData(TArrayView<FPushRequest> Requests) override;
 
diff --git a/Engine/Source/Developer/Virtualization/Private/VirtualizationFileBackend.cpp b/Engine/Source/Developer/Virtualization/Private/VirtualizationFileBackend.cpp
--- a/Engine/Source/Developer/Virtualization/Private/VirtualizationFileBackend.cpp
+++ b/Engine/Source/Developer/Virtualization/Private/VirtualizationFileBackend.cpp
@@ -53,6 +53,11 @@ bool FFileSystemBackend::Initialize(const FString& ConfigEntry)
 	return true;
 }
 
+IVirtualizationBackend::EConnectionStatus FFileSystemBackend::OnConnect()
+{
+	return IVirtualizationBackend::EConnectionStatus::Connected;
+}
+
 EPushResult FFileSystemBackend::PushData(const FIoHash& Id, const FCompressedBuffer& Payload, const FString& PackageContext)
 {
 	TRACE_CPUPROFILER_EVENT_SCOPE(FFileSystemBackend::PushData);
diff --git a/Engine/Source/Developer/Virtualization/Private/VirtualizationFileBackend.h b/Engine/Source/Developer/Virtualization/Private/VirtualizationFileBackend.h
--- a/Engine/Source/Developer/Virtualization/Private/VirtualizationFileBackend.h
+++ b/Engine/Source/Developer/Virtualization/Private/VirtualizationFileBackend.h
@@ -39,6 +39,8 @@ private:
 
 	virtual bool Initialize(const FString& ConfigEntry) override;
 
+	virtual EConnectionStatus OnConnect() override;
+
 	virtual EPushResult PushData(const FIoHash& Id, const FCompressedBuffer& Payload, const FString& PackageContext) override;
 
 	virtual FCompressedBuffer PullData(const FIoHash& Id) override;
diff --git a/Engine/Source/Developer/Virtualization/Private/VirtualizationHttpBackend.cpp b/Engine/Source/Developer/Virtualization/Private/VirtualizationHttpBackend.cpp
--- a/Engine/Source/Developer/Virtualization/Private/VirtualizationHttpBackend.cpp
+++ b/Engine/Source/Developer/Virtualization/Private/VirtualizationHttpBackend.cpp
@@ -1373,6 +1373,14 @@ bool FHttpBackend::Initialize(const FString& ConfigEntry)
 	return true;
 }
 
+IVirtualizationBackend::EConnectionStatus FHttpBackend::OnConnect()
+{
+	// TODO: In theory this backend has connection work that should be moved here, but we are more
+	// likely to just delete this backend at the moment, since it has been effectively replaced with
+	// the DDC backend (which can use horde storage)
+	return IVirtualizationBackend::EConnectionStatus::Connected;
+}
+
 EPushResult FHttpBackend::PushData(const FIoHash& Id, const FCompressedBuffer& CompressedPayload, const FString& PackageContext)
 {
 	using namespace Utility;
diff --git a/Engine/Source/Developer/Virtualization/Private/VirtualizationHttpBackend.h b/Engine/Source/Developer/Virtualization/Private/VirtualizationHttpBackend.h
--- a/Engine/Source/Developer/Virtualization/Private/VirtualizationHttpBackend.h
+++ b/Engine/Source/Developer/Virtualization/Private/VirtualizationHttpBackend.h
@@ -40,6 +40,8 @@ private:
 
 	virtual bool Initialize(const FString& ConfigEntry) override;
 
+	virtual EConnectionStatus OnConnect() override;
+
 	virtual EPushResult PushData(const FIoHash& Id, const FCompressedBuffer& CompressedPayload, const FString& PackageContext) override;
 
 	virtual FCompressedBuffer PullData(const FIoHash& Id) override;
diff --git a/Engine/Source/Developer/Virtualization/Private/VirtualizationManager.cpp b/Engine/Source/Developer/Virtualization/Private/VirtualizationManager.cpp
--- a/Engine/Source/Developer/Virtualization/Private/VirtualizationManager.cpp
+++ b/Engine/Source/Developer/Virtualization/Private/VirtualizationManager.cpp
@@ -25,6 +25,18 @@ namespace UE::Virtualization
 {
 UE_REGISTER_VIRTUALIZATION_SYSTEM(UE::Virtualization::FVirtualizationManager, Default);
 
+// Can be defined as 1 by programs target.cs files force the backend connections
+// to lazy initialize on first use rather than when the system is initialized.
+#ifndef UE_VIRTUALIZATION_CONNECTION_LAZY_INIT
+	#define UE_VIRTUALIZATION_CONNECTION_LAZY_INIT 0
+#endif //UE_VIRTUALIZATION_CONNECTION_LAZY_INIT
+
+// TODO: Move to RegisterConsoleCommands
+static TAutoConsoleVariable<bool> CVarLazyInitConnections(
+	TEXT("VA.LazyInitConnections"),
+	false,
+	TEXT("When true the VA backends will defer creating their connections until first use"));
+
 /** Utility struct, similar to FScopeLock but allows the lock to be enabled/disabled more easily */
 struct FConditionalScopeLock
 {
@@ -281,6 +293,8 @@ FVirtualizationManager::FVirtualizationManager()
 	, bFilterEnginePluginContent(true)
 	, bFilterMapContent(true)
 	, bAllowSubmitIfVirtualizationFailed(false)
+	, bLazyInitConnections(false)
+	, bPendingBackendConnections(false)
 {
 }
 
@@ -319,6 +333,9 @@ bool FVirtualizationManager::Initialize(const FInitParams& InitParams)
 	ProjectName = InitParams.ProjectName;
 
 	ApplySettingsFromConfigFiles(InitParams.ConfigFile);
+	ApplySettingsFromFromCmdline();
+	ApplySettingsFromCVar();
+
 	ApplyDebugSettingsFromFromCmdline();
 
 	// Do this after all of the command line settings have been processed and any 
@@ -470,6 +487,8 @@ bool FVirtualizationManager::PushData(TArrayView<FPushRequest> Requests, EStorag
 		return false;
 	}
 
+	EnsureBackendConnections();
+
 	FConditionalScopeLock _(&DebugValues.ForceSingleThreadedCS, DebugValues.bSingleThreaded);
 
 	// TODO: Note that all push operations are currently synchronous, probably 
@@ -482,6 +501,12 @@ bool FVirtualizationManager::PushData(TArrayView<FPushRequest> Requests, EStorag
 
 	for (IVirtualizationBackend* Backend : Backends)
 	{
+		if (Backend->GetConnectionStatus() != IVirtualizationBackend::EConnectionStatus::Connected)
+		{
+			UE_LOG(LogVirtualization, Verbose, TEXT("Cannot push to backend '%s' as it is not connected"), *Backend->GetDebugName());
+			continue;
+		}
+
 		const bool bResult = TryPushDataToBackend(*Backend, ValidatedRequests);
 
 		UE_CLOG(bResult == true, LogVirtualization, Verbose, TEXT("[%s] Pushed '%d' payload(s)"), *Backend->GetDebugName(), ValidatedRequests.Num());
@@ -541,6 +566,8 @@ FCompressedBuffer FVirtualizationManager::PullData(const FIoHash& Id)
 		return FCompressedBuffer();
 	}
 
+	EnsureBackendConnections();
+
 	FConditionalScopeLock _(&DebugValues.ForceSingleThreadedCS, DebugValues.bSingleThreaded);
 
 	GetNotificationEvent().Broadcast(IVirtualizationSystem::PullBegunNotification, Id);
@@ -592,11 +619,19 @@ EQueryResult FVirtualizationManager::QueryPayloadStatuses(TArrayView<const FIoHa
 	HitCount.SetNum(Ids.Num());
 	Results.SetNum(Ids.Num());
 
+	EnsureBackendConnections();
+
 	{
 		FConditionalScopeLock _(&DebugValues.ForceSingleThreadedCS, DebugValues.bSingleThreaded);
 
 		for (IVirtualizationBackend* Backend : Backends)
 		{
+			if (Backend->GetConnectionStatus() != IVirtualizationBackend::EConnectionStatus::Connected)
+			{
+				UE_LOG(LogVirtualization, Verbose, TEXT("Cannot query backend '%s' as it is not connected"), *Backend->GetDebugName());
+				continue;
+			}
+
 			if (!Backend->DoPayloadsExist(Ids, Results))
 			{
 				// If a backend entirely failed we should early out and report the problem
@@ -906,6 +941,22 @@ void FVirtualizationManager::ApplySettingsFromConfigFiles(const FConfigFile& Con
 		UE_LOG(LogVirtualization, Error, TEXT("Failed to load [Core.VirtualizationModule].AllowSubmitIfVirtualizationFailed from config file!"));
 	}
 
+	// Optional
+#if UE_VIRTUALIZATION_CONNECTION_LAZY_INIT == 0
+	bool bLazyInitConnectionsFromIni = true;
+	if (ConfigFile.GetBool(ConfigSection, TEXT("LazyInitConnections"), bLazyInitConnectionsFromIni))
+	{
+		bLazyInitConnections = bLazyInitConnectionsFromIni;
+		UE_LOG(LogVirtualization, Display, TEXT("\tLazyInitConnections : %s"), bLazyInitConnections ? TEXT("true") : TEXT("false"));
+	}
+	else
+	{
+		UE_LOG(LogVirtualization, Error, TEXT("Failed to load [Core.VirtualizationModule].LazyInitConnections from config file!"));
+	}
+#else
+	UE_LOG(LogVirtualization, Display, TEXT("\tLazyInitConnections : false (set by code)"), bLazyInitConnections ? TEXT("true") : TEXT("false"));
+#endif //UE_VIRTUALIZATION_CONNECTION_LAZY_INIT
+
 	// Check for any legacy settings and print them out (easier to do this in one block rather than one and time)
 	{
 		// Entries that are allows to be in [Core.ContentVirtualization
@@ -936,6 +987,24 @@ void FVirtualizationManager::ApplySettingsFromConfigFiles(const FConfigFile& Con
 	}
 }
 
+void FVirtualizationManager::ApplySettingsFromFromCmdline()
+{
+	if (!bLazyInitConnections && FParse::Param(FCommandLine::Get(), TEXT("VA-LazyInitConnections")))
+	{
+		bLazyInitConnections = true;
+		UE_LOG(LogVirtualization, Display, TEXT("Cmdline has set the virtualization system backends to lazy init their connections"));
+	}
+}
+
+void FVirtualizationManager::ApplySettingsFromCVar()
+{
+	if (!bLazyInitConnections && CVarLazyInitConnections.GetValueOnAnyThread())
+	{
+		bLazyInitConnections = true;
+		UE_LOG(LogVirtualization, Display, TEXT("CVar has set the virtualization system backends to lazy init their connections"));
+	}
+}
+
 void FVirtualizationManager::ApplyDebugSettingsFromFromCmdline()
 {
 	if (FParse::Param(FCommandLine::Get(), TEXT("VA-SingleThreaded")))
@@ -1340,6 +1409,7 @@ void FVirtualizationManager::AddBackend(TUniquePtr<IVirtualizationBackend> Backe
 
 	// Get a reference pointer to use in the other backend arrays
 	IVirtualizationBackend* BackendRef = AllBackends.Last().Get();
+	check(BackendRef != nullptr);
 
 	if (BackendRef->IsOperationSupported(IVirtualizationBackend::EOperations::Pull))
 	{
@@ -1351,13 +1421,46 @@ void FVirtualizationManager::AddBackend(TUniquePtr<IVirtualizationBackend> Backe
 		PushArray.Add(BackendRef);
 	}
 
+	// We immediately try to connect once the backend has been added.
+	// In the future this will be made async to avoid blocking the GameThread on startup
+	if (!bLazyInitConnections)
+	{
+		BackendRef->Connect();
+	}
+	else
+	{
+		bPendingBackendConnections = true;
+	}
+
 	COOK_STAT(Profiling::CreateStats(*BackendRef));
 }
 
+void FVirtualizationManager::EnsureBackendConnections()
+{
+	if (bPendingBackendConnections)
+	{
+		// Only allow one thread to initialize the system at a time
+		static FCriticalSection InitCS;
+
+		FScopeLock _(&InitCS);
+		if (bPendingBackendConnections)
+		{
+			for (const TUniquePtr<IVirtualizationBackend>& Backend : AllBackends)
+			{
+				Backend->Connect();
+			}
+
+			bPendingBackendConnections = false;
+		}
+	}
+}
+
 void FVirtualizationManager::CachePayload(const FIoHash& Id, const FCompressedBuffer& Payload, const IVirtualizationBackend* BackendSource)
 {
 	TRACE_CPUPROFILER_EVENT_SCOPE(FVirtualizationManager::CachePayload);
 
+	check(!bPendingBackendConnections);
+
 	// We start caching at the first (assumed to be fastest) local cache backend. 
 	for (IVirtualizationBackend* BackendToCache : CacheStorageBackends)
 	{
@@ -1444,6 +1547,12 @@ FCompressedBuffer FVirtualizationManager::PullDataFromAllBackends(const FIoHash&
 			continue;
 		}
 
+		if (Backend->GetConnectionStatus() != IVirtualizationBackend::EConnectionStatus::Connected)
+		{
+			UE_LOG(LogVirtualization, Verbose, TEXT("Cannot pull from backend '%s' as it is not connected"), *Backend->GetDebugName());
+			continue;
+		}
+
 		FCompressedBuffer Payload = PullDataFromBackend(*Backend, Id);
 
 		if (Payload)
diff --git a/Engine/Source/Developer/Virtualization/Private/VirtualizationManager.h b/Engine/Source/Developer/Virtualization/Private/VirtualizationManager.h
--- a/Engine/Source/Developer/Virtualization/Private/VirtualizationManager.h
+++ b/Engine/Source/Developer/Virtualization/Private/VirtualizationManager.h
@@ -61,6 +61,7 @@ class FOutputDevice;
  * FilterEngineContent=True/False			When true any payload from a package under Engine/Content/.. will be excluded from virtualization
  * FilterEnginePluginContent=True/False		When true any payload from a package under Engine/Plugins/../Content/.. will be excluded from virtualization
  * FilterMapContent=True/False				When true any payload stored in a .umap or _BuildData.uasset file will be excluded from virtualization
+ * LazyInitConnections=True/False			When true the backends will not try to make their connections until first used.
  * 
  * PackagePath Setup:
  * 
@@ -142,6 +143,9 @@ private:
 private:
 
 	void ApplySettingsFromConfigFiles(const FConfigFile& ConfigFile);
+	void ApplySettingsFromFromCmdline();
+	void ApplySettingsFromCVar();
+
 	void ApplyDebugSettingsFromFromCmdline();
 
 	void RegisterConsoleCommands();
@@ -161,6 +165,8 @@ private:
 
 	void AddBackend(TUniquePtr<IVirtualizationBackend> Backend, FBackendArray& PushArray);
 
+	void EnsureBackendConnections();
+
 	void CachePayload(const FIoHash& Id, const FCompressedBuffer& Payload, const IVirtualizationBackend* BackendSource);
 
 	bool TryCacheDataToBackend(IVirtualizationBackend& Backend, const FIoHash& Id, const FCompressedBuffer& Payload);
@@ -231,6 +237,9 @@ private:
 	/** Should file submits be allowed to continue if a call to TryVirtualizePackages fails */
 	bool bAllowSubmitIfVirtualizationFailed;
 	
+	/** Should backends defer connecting to their services until first use */
+	bool bLazyInitConnections;
+
 private:
 
 	/** The name of the current project */
@@ -254,6 +263,9 @@ private:
 	 */
 	FBackendArray PullEnabledBackends;
 
+	/** Do we have backends that have not yet tried connecting to their services */
+	bool bPendingBackendConnections;
+
 	/** Our notification Event */
 	FOnNotification NotificationEvent;
 
diff --git a/Engine/Source/Developer/Virtualization/Private/VirtualizationSourceControlBackend.cpp b/Engine/Source/Developer/Virtualization/Private/VirtualizationSourceControlBackend.cpp
--- a/Engine/Source/Developer/Virtualization/Private/VirtualizationSourceControlBackend.cpp
+++ b/Engine/Source/Developer/Virtualization/Private/VirtualizationSourceControlBackend.cpp
@@ -233,6 +233,11 @@ bool FSourceControlBackend::Initialize(const FString& ConfigEntry)
 		return false;
 	}
 
+	return true;
+}
+
+IVirtualizationBackend::EConnectionStatus FSourceControlBackend::OnConnect()
+{
 	// We do not want the connection to have a client workspace so explicitly set it to empty
 	FSourceControlInitSettings SCCSettings(FSourceControlInitSettings::EBehavior::OverrideExisting);
 	SCCSettings.AddSetting(TEXT("P4Client"), TEXT(""));
@@ -241,11 +246,11 @@ bool FSourceControlBackend::Initialize(const FString& ConfigEntry)
 	if (!SCCProvider.IsValid())
 	{
 		UE_LOG(LogVirtualization, Error, TEXT("[%s] Failed to create a perforce connection, this seems to be unsupported by the editor"), *GetDebugName());
-		return false;
+		return IVirtualizationBackend::EConnectionStatus::Error;
 	}
 
 	SCCProvider->Init(true);
- 
+
 	if (IsInGameThread())
 	{
 		// Note that if the connect is failing then we expect it to fail here rather than in the subsequent attempts to get the meta info file
@@ -263,7 +268,7 @@ bool FSourceControlBackend::Initialize(const FString& ConfigEntry)
 				Errors.ToText()));
 
 			OnConnectionError();
-			return true;
+			return IVirtualizationBackend::EConnectionStatus::Error;
 		}
 	}
 
@@ -278,11 +283,11 @@ bool FSourceControlBackend::Initialize(const FString& ConfigEntry)
 		FMessageLog Log("LogVirtualization");
 
 		Log.Warning(FText::Format(LOCTEXT("FailedMetaInfo", "Failed to find 'payload_metainfo.txt' in the depot '{0}'\nThe source control backend will be unable to pull payloads, is your source control  config set up correctly?"),
-					FText::FromString(DepotRoot)));
-		
+			FText::FromString(DepotRoot)));
+
 		OnConnectionError();
-		return true;
-	}	
+		return IVirtualizationBackend::EConnectionStatus::Error;
+}
 #else
 	TSharedRef<FDownloadFile, ESPMode::ThreadSafe> DownloadCommand = ISourceControlOperation::Create<FDownloadFile>();
 	if (!SCCProvider->TryToDownloadFileFromBackgroundThread(DownloadCommand, PayloadMetaInfoPath))
@@ -290,11 +295,11 @@ bool FSourceControlBackend::Initialize(const FString& ConfigEntry)
 		FMessageLog Log("LogVirtualization");
 
 		Log.Warning(FText::Format(LOCTEXT("FailedMetaInfo", "Failed to find 'payload_metainfo.txt' in the depot '{0}'\nThe source control backend will be unable to pull payloads, is your source control  config set up correctly?"),
-					FText::FromString(DepotRoot)));
-		
+			FText::FromString(DepotRoot)));
+
 		OnConnectionError();
-		return true;
-	}		
+		return IVirtualizationBackend::EConnectionStatus::Error;
+	}
 #endif //IS_SOURCE_CONTROL_THREAD_SAFE
 
 	FSharedBuffer MetaInfoBuffer = DownloadCommand->GetFileData(PayloadMetaInfoPath);
@@ -303,18 +308,17 @@ bool FSourceControlBackend::Initialize(const FString& ConfigEntry)
 		FMessageLog Log("LogVirtualization");
 
 		Log.Warning(FText::Format(LOCTEXT("FailedMetaInfo", "Failed to find 'payload_metainfo.txt' in the depot '{0}'\nThe source control backend will be unable to pull payloads, is your source control  config set up correctly?"),
-					FText::FromString(DepotRoot)));
+			FText::FromString(DepotRoot)));
 
 		OnConnectionError();
-		return true;
+		return IVirtualizationBackend::EConnectionStatus::Error;
 	}
 
 	// Currently we do not do anything with the payload meta info, in the future we could structure
 	// it's format to include more information that might be worth logging or something. 
 	// But for now being able to pull the payload meta info path at least shows that we can use the
 	// depot.
-
-	return true;
+	return IVirtualizationBackend::EConnectionStatus::Connected;
 }
 
 FCompressedBuffer FSourceControlBackend::PullData(const FIoHash& Id)
diff --git a/Engine/Source/Developer/Virtualization/Private/VirtualizationSourceControlBackend.h b/Engine/Source/Developer/Virtualization/Private/VirtualizationSourceControlBackend.h
--- a/Engine/Source/Developer/Virtualization/Private/VirtualizationSourceControlBackend.h
+++ b/Engine/Source/Developer/Virtualization/Private/VirtualizationSourceControlBackend.h
@@ -65,6 +65,8 @@ private:
 	/* IVirtualizationBackend implementation */
 
 	virtual bool Initialize(const FString& ConfigEntry) override;
+
+	virtual EConnectionStatus OnConnect() override;
 	
 	virtual EPushResult PushData(const FIoHash& Id, const FCompressedBuffer& Payload, const FString& Context) override;
 	virtual bool PushData(TArrayView<FPushRequest> Requests) override;
diff --git a/Engine/Source/Developer/Virtualization/Public/IVirtualizationBackend.h b/Engine/Source/Developer/Virtualization/Public/IVirtualizationBackend.h
--- a/Engine/Source/Developer/Virtualization/Public/IVirtualizationBackend.h
+++ b/Engine/Source/Developer/Virtualization/Public/IVirtualizationBackend.h
@@ -37,7 +37,7 @@ enum class EPushResult
  */
 class IVirtualizationBackend
 {
-public: 
+public:
 	/** Enum detailing which operations a backend can support */
 	enum class EOperations : uint8
 	{
@@ -51,8 +51,19 @@ public:
 
 	FRIEND_ENUM_CLASS_FLAGS(EOperations);
 
+	/** Status of the backends connection to its services */
+	enum class EConnectionStatus
+	{
+		/** No connection attempt has been made */
+		None,
+		/** The previous connection attempt ended with an error */
+		Error,
+		/** The connection has been made successfully */
+		Connected,
+	};
+
 protected:
-	
+
 	IVirtualizationBackend(FStringView InConfigName, FStringView InDebugName, EOperations InSupportedOperations)
 		: SupportedOperations(InSupportedOperations)
 		, DebugDisabledOperations(EOperations::None)
@@ -69,16 +80,37 @@ public:
 	 * This will be called during the setup of the backend hierarchy. The entry config file
 	 * entry that caused the backend to be created will be passed to the method so that any
 	 * additional settings may be parsed from it.
-	 * Take care to clearly log any error that occurs so that the end user has a clear way 
+	 * Take care to clearly log any error that occurs so that the end user has a clear way
 	 * to fix them.
-	 * 
-	 * @param ConfigEntry	The entry for the backend from the config ini file that may 
+	 *
+	 * @param ConfigEntry	The entry for the backend from the config ini file that may
 	 *						contain additional settings.
-	 * @return				Returning false indicates that initialization failed in a way 
+	 * @return				Returning false indicates that initialization failed in a way
 	 *						that the backend will not be able to function correctly.
 	 */
 	virtual bool Initialize(const FString& ConfigEntry) = 0;
 
+	/**
+	 * Attempt to connect the backend to its services if not already connected.
+	 */
+	virtual void Connect()
+	{
+		if (ConnectionStatus != EConnectionStatus::Connected)
+		{
+			ConnectionStatus = OnConnect();
+		}
+	}
+
+	/**
+	 * Returns the connection status of the backend
+	 * 
+	 * @return @see EConnectionStatus
+	 */
+	EConnectionStatus GetConnectionStatus() const
+	{
+		return ConnectionStatus;
+	}
+
 	/**
 	 * The backend will attempt to store the given payload by what ever method the backend uses.
 	 * NOTE: It is assumed that the virtualization manager will run all appropriate validation
@@ -208,6 +240,11 @@ public:
 		return DebugName;
 	}
 
+private:
+
+	/** Override to implement the backends connection code */
+	virtual EConnectionStatus OnConnect() = 0;
+
 private:
 
 	/** The operations that this backend supports */
@@ -220,6 +257,9 @@ private:
 
 	/** Combination of the backend type and the name used to create it in the virtualization graph */
 	FString DebugName;
+
+	/** The status of the connection to the backends service */
+	EConnectionStatus ConnectionStatus = EConnectionStatus::None;
 };
 
 /** 
