Add commandline options and cvars to disable the virtualization process

#rnx

- Added a cvar 'VA.DisableVirtualization' and a commandline '-VA-DisableVirtualization-' that when set are equivalent to setting '[Core.VirtualizationModule].EnablePayloadVirtualization=false' and will. prevent the VA system from virtualizing any packages on submit.

diff --git a/Engine/Source/Developer/Virtualization/Private/VirtualizationManager.cpp b/Engine/Source/Developer/Virtualization/Private/VirtualizationManager.cpp
--- a/Engine/Source/Developer/Virtualization/Private/VirtualizationManager.cpp
+++ b/Engine/Source/Developer/Virtualization/Private/VirtualizationManager.cpp
@@ -37,6 +37,12 @@ static TAutoConsoleVariable<bool> CVarLazyInitConnections(
 	false,
 	TEXT("When true the VA backends will defer creating their connections until first use"));
 
+// TODO: Move to RegisterConsoleCommands
+static TAutoConsoleVariable<bool> CVarDisableVirtualization(
+	TEXT("VA.DisableVirtualization"),
+	false,
+	TEXT("When true submitting packages in the editor will no longer trigger the virtualization process (only applied when the system is initialized)"));
+
 /** Utility struct, similar to FScopeLock but allows the lock to be enabled/disabled more easily */
 struct FConditionalScopeLock
 {
@@ -994,6 +1000,12 @@ void FVirtualizationManager::ApplySettingsFromFromCmdline()
 		bLazyInitConnections = true;
 		UE_LOG(LogVirtualization, Display, TEXT("Cmdline has set the virtualization system backends to lazy init their connections"));
 	}
+
+	if (bEnablePayloadVirtualization && FParse::Param(FCommandLine::Get(), TEXT("VA-DisableVirtualization")))
+	{
+		bEnablePayloadVirtualization = false;
+		UE_LOG(LogVirtualization, Warning, TEXT("Cmdline has disabled the virtualization process"));
+	}
 }
 
 void FVirtualizationManager::ApplySettingsFromCVar()
@@ -1003,6 +1015,12 @@ void FVirtualizationManager::ApplySettingsFromCVar()
 		bLazyInitConnections = true;
 		UE_LOG(LogVirtualization, Display, TEXT("CVar has set the virtualization system backends to lazy init their connections"));
 	}
+
+	if (bEnablePayloadVirtualization && CVarDisableVirtualization.GetValueOnAnyThread())
+	{
+		bEnablePayloadVirtualization = false;
+		UE_LOG(LogVirtualization, Display, TEXT("CVar has disabled the virtualization process"));
+	}
 }
 
 void FVirtualizationManager::ApplyDebugSettingsFromFromCmdline()
