Improve the error message given when the VA perforce backend cannot connect because the PerforceSourceControl plugin is disabled

#rnx

- We could try force loading and enabling the plugin at this point but I am not a fan of the idea of overriding the set up of the target in a way that is not immediately obvious to the end user.
- The virtualization module needs to include the 'Projects' module to be able to access the plugin system directly and while editing the file I took the opportunity to clean up how the module is declaring its dependencies.

diff --git a/Engine/Source/Developer/Virtualization/Private/VirtualizationSourceControlBackend.cpp b/Engine/Source/Developer/Virtualization/Private/VirtualizationSourceControlBackend.cpp
--- a/Engine/Source/Developer/Virtualization/Private/VirtualizationSourceControlBackend.cpp
+++ b/Engine/Source/Developer/Virtualization/Private/VirtualizationSourceControlBackend.cpp
@@ -17,6 +17,7 @@
 #include "SourceControlOperations.h"
 #include "VirtualizationSourceControlUtilities.h"
 #include "VirtualizationUtilities.h"
+#include "Interfaces/IPluginManager.h"
 
 // When the SourceControl module (or at least the perforce source control module) is thread safe we
 // can enable this and stop using the hacky work around 'TryToDownloadFileFromBackgroundThread'
@@ -238,6 +239,26 @@ bool FSourceControlBackend::Initialize(const FString& ConfigEntry)
 
 IVirtualizationBackend::EConnectionStatus FSourceControlBackend::OnConnect()
 {
+	// First check that the 'PerforceSourceControl' plugin exists and is enabled so that we can give specific
+	// error messages if not.
+	// We could try to force enable the plugin at this point, in the same way that the stand alone tool does
+	// but it is probably better to inform the user and have them add/enable the plugin for their target explicitly
+	// rather than have us do it behind the scenes.
+	// This is only expected to be a problem when first enabling VA for a project/target and not something a user
+	// will experience day to day.
+	TSharedPtr<IPlugin> P4Plugin = IPluginManager::Get().FindPlugin("PerforceSourceControl");
+	if (!P4Plugin)
+	{
+		UE_LOG(LogVirtualization, Error, TEXT("[%s] Failed to find the 'PerforceSourceControl' plugin"), *GetDebugName());
+		return IVirtualizationBackend::EConnectionStatus::Error;
+	}
+
+	if (!P4Plugin->IsEnabled())
+	{
+		UE_LOG(LogVirtualization, Error, TEXT("[%s] The 'PerforceSourceControl' plugin is disabled for this target"), *GetDebugName());
+		return IVirtualizationBackend::EConnectionStatus::Error;
+	}
+
 	// We do not want the connection to have a client workspace so explicitly set it to empty
 	FSourceControlInitSettings SCCSettings(FSourceControlInitSettings::EBehavior::OverrideExisting);
 	SCCSettings.AddSetting(TEXT("P4Client"), TEXT(""));
@@ -245,7 +266,7 @@ IVirtualizationBackend::EConnectionStatus FSourceControlBackend::OnConnect()
 	SCCProvider = ISourceControlModule::Get().CreateProvider(FName("Perforce"), TEXT("Virtualization"), SCCSettings);
 	if (!SCCProvider.IsValid())
 	{
-		UE_LOG(LogVirtualization, Error, TEXT("[%s] Failed to create a perforce connection, this seems to be unsupported by the editor"), *GetDebugName());
+		UE_LOG(LogVirtualization, Error, TEXT("[%s] Failed to create a perforce source control provider"), *GetDebugName());
 		return IVirtualizationBackend::EConnectionStatus::Error;
 	}
 
diff --git a/Engine/Source/Developer/Virtualization/Virtualization.Build.cs b/Engine/Source/Developer/Virtualization/Virtualization.Build.cs
--- a/Engine/Source/Developer/Virtualization/Virtualization.Build.cs
+++ b/Engine/Source/Developer/Virtualization/Virtualization.Build.cs
@@ -7,17 +7,20 @@ public class Virtualization : ModuleRules
 {
 	public Virtualization(ReadOnlyTargetRules Target) : base(Target)
 	{
-		PrivateDependencyModuleNames.Add("Core");
-		PrivateDependencyModuleNames.Add("CoreUObject");
-		PrivateDependencyModuleNames.Add("MessageLog");
+		PrivateDependencyModuleNames.AddRange(
+			new string[]
+			{
+				"Core",
+				"CoreUObject",
+				"ApplicationCore",
+				"DerivedDataCache",
+				"MessageLog",
+				"Projects",
+				"SourceControl"
+			});
 
-		// Dependency for DDC2
-		PrivateDependencyModuleNames.Add("DerivedDataCache");
-
-		// Dependency for source control backend
-		PrivateDependencyModuleNames.Add("SourceControl");
-
-		// Dependencies for the Horde Storage backend
+		// The dependencies below this point are for FHttpBackend
+		// and can be removed when it is removed
 		PrivateDependencyModuleNames.AddRange(new string[] { "SSL", "Json" });
 
 		AddEngineThirdPartyPrivateStaticDependencies(Target, "libcurl");
