Remove the FHttpBackend

#rnx

- The backend is no longer used as we access horde storage via the ddc backend.

diff --git a/Engine/Source/Developer/Virtualization/Private/VirtualizationHttpBackend.cpp b/Engine/Source/Developer/Virtualization/Private/VirtualizationHttpBackend.cpp
deleted file mode 100644
diff --git a/Engine/Source/Developer/Virtualization/Private/VirtualizationHttpBackend.h b/Engine/Source/Developer/Virtualization/Private/VirtualizationHttpBackend.h
deleted file mode 100644
diff --git a/Engine/Source/Developer/Virtualization/Virtualization.Build.cs b/Engine/Source/Developer/Virtualization/Virtualization.Build.cs
--- a/Engine/Source/Developer/Virtualization/Virtualization.Build.cs
+++ b/Engine/Source/Developer/Virtualization/Virtualization.Build.cs
@@ -18,12 +18,5 @@ public class Virtualization : ModuleRules
 				"Projects",
 				"SourceControl"
 			});
-
-		// The dependencies below this point are for FHttpBackend
-		// and can be removed when it is removed
-		PrivateDependencyModuleNames.AddRange(new string[] { "SSL", "Json" });
-
-		AddEngineThirdPartyPrivateStaticDependencies(Target, "libcurl");
-		AddEngineThirdPartyPrivateStaticDependencies(Target, "OpenSSL");
 	}
 }
