Refactor the VA pushing API so that we only need to implement a single path

#rnx

### VirtualizationSystem/VirtualizationManager
- Added an overload of ::Push taking just one FPUshRequest so that people don't have to keep adding MakeArrayView boiler plate when pushing a single request
- Change the order of the last two parameters for the raw ::Push call as this will group all of the payload specific parameters together and leave the target storage type at the end. It is unlikely that anyone is calling the older version but it has been deprecated for safety.

### IVirtualizationBackend
- The none FPushRequest overload of ::PushData is no longer virtual, it just converts the parameters to FPushRequest and calls that overload instead. In this way we now only have one pushing code path in our backends. We could probably look into removing this overload at this point (since the higher level IVirtualizationSystem will now convert all push requests into a FPushRequest form) but it is not considered worth it at the moment when the simple overload covers our needs.
- Removed EPushResult in favour of just returning true/false for the overall operation.If the caller needs a more detailed breakdown then they will have to use an overload that takes an FPushRequest over raw parameters.
-- At the moment FPushRequest does not contain a full breakdown of what happened, so with this submit we are effectively losing the ability to find out if the payload was already in the backends or not, however the batch version of push was already not returning this info so it is not a big loss. Fixing FPushRequest to return a better break down of what happened will be done in UE-160942
- Removed the none batch Push paths from the source control and ddc backends as they already supported batch pushing.
- File backend needed to be converted to supporting batch pushing, which is pretty much the same code as before except we need to iterate over the container of FPushRequests.
-- The backend does not early out on error as it tends to be quite fast. We might want to consider an official policy for the VA system, if we should early out of errors or not.

diff --git a/Engine/Source/Developer/Virtualization/Private/VirtualizationDDCBackend.cpp b/Engine/Source/Developer/Virtualization/Private/VirtualizationDDCBackend.cpp
--- a/Engine/Source/Developer/Virtualization/Private/VirtualizationDDCBackend.cpp
+++ b/Engine/Source/Developer/Virtualization/Private/VirtualizationDDCBackend.cpp
@@ -74,51 +74,6 @@ IVirtualizationBackend::EConnectionStatus FDDCBackend::OnConnect()
 	return IVirtualizationBackend::EConnectionStatus::Connected;
 }
 
-EPushResult FDDCBackend::PushData(const FIoHash& Id, const FCompressedBuffer& Payload, const FString& PackageContext)
-{
-	TRACE_CPUPROFILER_EVENT_SCOPE(FDDCBackend::PushData);
-
-	if (DoesPayloadExist(Id))
-	{
-		UE_LOG(LogVirtualization, Verbose, TEXT("[%s] Already has a copy of the payload '%s'."), *GetDebugName(), *LexToString(Id));
-		return EPushResult::PayloadAlreadyExisted;
-	}
-
-	UE::DerivedData::ICache& Cache = UE::DerivedData::GetCache();
-	
-	UE::DerivedData::FCacheKey Key;
-	Key.Bucket = Bucket;
-	Key.Hash = Id;
-
-	UE::DerivedData::FValue DerivedDataValue(Payload);
-	check(DerivedDataValue.GetRawHash() == Id);
-
-	UE::DerivedData::FCacheRecordBuilder RecordBuilder(Key);
-	RecordBuilder.AddValue(ToDerivedDataValueId(Id), DerivedDataValue);
-
-	UE::DerivedData::FRequestOwner Owner(UE::DerivedData::EPriority::Blocking);
-
-	UE::DerivedData::FCachePutResponse Result;
-	auto Callback = [&Result](UE::DerivedData::FCachePutResponse&& Response)
-	{
-		Result = Response;
-	};
-
-	// TODO: Improve the name when we start passing more context to this function
-	Cache.Put({{{TEXT("Mirage")}, RecordBuilder.Build(), TransferPolicy}}, Owner, MoveTemp(Callback));
-
-	Owner.Wait();
-
-	if (Result.Status == UE::DerivedData::EStatus::Ok)
-	{
-		return EPushResult::Success;
-	}
-	else
-	{
-		return EPushResult::Failed;
-	}
-}
-
 bool FDDCBackend::PushData(TArrayView<FPushRequest> Requests)
 {
 	TRACE_CPUPROFILER_EVENT_SCOPE(FDDCBackend::PushData);
diff --git a/Engine/Source/Developer/Virtualization/Private/VirtualizationDDCBackend.h b/Engine/Source/Developer/Virtualization/Private/VirtualizationDDCBackend.h
--- a/Engine/Source/Developer/Virtualization/Private/VirtualizationDDCBackend.h
+++ b/Engine/Source/Developer/Virtualization/Private/VirtualizationDDCBackend.h
@@ -38,7 +38,6 @@ private:
 
 	virtual EConnectionStatus OnConnect() override;
 
-	virtual EPushResult PushData(const FIoHash& Id, const FCompressedBuffer& Payload, const FString& PackageContext) override;
 	virtual bool PushData(TArrayView<FPushRequest> Requests) override;
 
 	virtual FCompressedBuffer PullData(const FIoHash& Id) override;
diff --git a/Engine/Source/Developer/Virtualization/Private/VirtualizationFileBackend.cpp b/Engine/Source/Developer/Virtualization/Private/VirtualizationFileBackend.cpp
--- a/Engine/Source/Developer/Virtualization/Private/VirtualizationFileBackend.cpp
+++ b/Engine/Source/Developer/Virtualization/Private/VirtualizationFileBackend.cpp
@@ -58,93 +58,109 @@ IVirtualizationBackend::EConnectionStatus FFileSystemBackend::OnConnect()
 	return IVirtualizationBackend::EConnectionStatus::Connected;
 }
 
-EPushResult FFileSystemBackend::PushData(const FIoHash& Id, const FCompressedBuffer& Payload, const FString& PackageContext)
+bool FFileSystemBackend::PushData(TArrayView<FPushRequest> Requests)
 {
 	TRACE_CPUPROFILER_EVENT_SCOPE(FFileSystemBackend::PushData);
 
-	if (DoesPayloadExist(Id))
-	{
-		UE_LOG(LogVirtualization, Verbose, TEXT("[%s] Already has a copy of the payload '%s'."), *GetDebugName(), *LexToString(Id));
-		return EPushResult::PayloadAlreadyExisted;
-	}
-
-	// Make sure to log any disk write failures to the user, even if this backend will often be optional as they are
-	// not expected and could indicate bigger problems.
-	// 
-	// First we will write out the payload to a temp file, after which we will move it to the correct storage location
-	// this helps reduce the chance of leaving corrupted data on disk in the case of a power failure etc.
-	const FString TempFilePath = FPaths::CreateTempFilename(*FPaths::ProjectSavedDir(), TEXT("miragepayload"));
-	
-	TUniquePtr<FArchive> FileAr(IFileManager::Get().CreateFileWriter(*TempFilePath));
+	int32 ErrorCount = 0;
 
-	if (FileAr == nullptr)
+	for (FPushRequest& Request : Requests)
 	{
-		TStringBuilder<MAX_SPRINTF> SystemErrorMsg;
-		Utils::GetFormattedSystemError(SystemErrorMsg);
+		const FIoHash& PayloadId = Request.GetIdentifier();
 
-		UE_LOG(LogVirtualization, Error, TEXT("[%s] Failed to write payload '%s' to '%s' due to system error: %s"), 
-			*GetDebugName(),
-			*LexToString(Id),
-			*TempFilePath,
-			SystemErrorMsg.ToString());
+		if (DoesPayloadExist(PayloadId))
+		{
+			UE_LOG(LogVirtualization, Verbose, TEXT("[%s] Already has a copy of the payload '%s'."), *GetDebugName(), *LexToString(PayloadId));
+			Request.SetStatus(FPushRequest::EStatus::Success);
 
-		return EPushResult::Failed;
-	}
+			continue;
+		}
 
-	for (const FSharedBuffer& Buffer : Payload.GetCompressed().GetSegments())
-	{
-		// Const cast because FArchive requires a non-const pointer!
-		FileAr->Serialize(const_cast<void*>(Buffer.GetData()), static_cast<int64>(Buffer.GetSize()));
-	}
+		// Make sure to log any disk write failures to the user, even if this backend will often be optional as they are
+		// not expected and could indicate bigger problems.
+		// 
+		// First we will write out the payload to a temp file, after which we will move it to the correct storage location
+		// this helps reduce the chance of leaving corrupted data on disk in the case of a power failure etc.
+		const FString TempFilePath = FPaths::CreateTempFilename(*FPaths::ProjectSavedDir(), TEXT("vapayload"));
 
-	if (!FileAr->Close())
-	{
-		TStringBuilder<MAX_SPRINTF> SystemErrorMsg;
-		Utils::GetFormattedSystemError(SystemErrorMsg);
+		TUniquePtr<FArchive> FileAr(IFileManager::Get().CreateFileWriter(*TempFilePath));
 
-		UE_LOG(LogVirtualization, Error, TEXT("[%s] Failed to write payload '%s' contents to '%s' due to system error: %s"),
-			*GetDebugName(),
-			*LexToString(Id),
-			*TempFilePath,
-			SystemErrorMsg.ToString());
+		if (FileAr == nullptr)
+		{
+			TStringBuilder<MAX_SPRINTF> SystemErrorMsg;
+			Utils::GetFormattedSystemError(SystemErrorMsg);
 
-		IFileManager::Get().Delete(*TempFilePath, true, false, true);  // Clean up the temp file if it is still around but do not failure cases to the user
-		
-		return EPushResult::Failed;
-	}
+			UE_LOG(LogVirtualization, Error, TEXT("[%s] Failed to write payload '%s' to '%s' due to system error: %s"),
+				*GetDebugName(),
+				*LexToString(PayloadId),
+				*TempFilePath,
+				SystemErrorMsg.ToString());
 
-	TStringBuilder<512> FilePath;
-	CreateFilePath(Id, FilePath);
+			ErrorCount++;
+			Request.SetStatus(FPushRequest::EStatus::Failed);
 
-	// If the file already exists we don't need to replace it, we will also do our own error logging.
-	if (!IFileManager::Get().Move(FilePath.ToString(), *TempFilePath, /*Replace*/ false, /*EvenIfReadOnly*/ false, /*Attributes*/ false, /*bDoNotRetryOrError*/ true))
-	{
-		// Store the error message in case we need to display it
-		TStringBuilder<MAX_SPRINTF> SystemErrorMsg;
-		Utils::GetFormattedSystemError(SystemErrorMsg);
+			continue;
+		}
 
-		IFileManager::Get().Delete(*TempFilePath, true, false, true); // Clean up the temp file if it is still around but do not failure cases to the user
+		for (const FSharedBuffer& Buffer : Request.GetPayload().GetCompressed().GetSegments())
+		{
+			// Const cast because FArchive requires a non-const pointer!
+			FileAr->Serialize(const_cast<void*>(Buffer.GetData()), static_cast<int64>(Buffer.GetSize()));
+		}
 
-		// Check if another thread or process was writing out the payload at the same time, if so we 
-		// don't need to give an error message.
-		if (DoesPayloadExist(Id))
+		if (!FileAr->Close())
 		{
-			UE_LOG(LogVirtualization, Verbose, TEXT("[%s] Already has a copy of the payload '%s'."), *GetDebugName(), *LexToString(Id));
-			return EPushResult::PayloadAlreadyExisted;
+			TStringBuilder<MAX_SPRINTF> SystemErrorMsg;
+			Utils::GetFormattedSystemError(SystemErrorMsg);
+
+			UE_LOG(LogVirtualization, Error, TEXT("[%s] Failed to write payload '%s' contents to '%s' due to system error: %s"),
+				*GetDebugName(),
+				*LexToString(PayloadId),
+				*TempFilePath,
+				SystemErrorMsg.ToString());
+
+			IFileManager::Get().Delete(*TempFilePath, true, false, true);  // Clean up the temp file if it is still around but do not failure cases to the user
+
+			ErrorCount++;
+			Request.SetStatus(FPushRequest::EStatus::Failed);
+
+			continue;
 		}
-		else
+
+		TStringBuilder<512> FilePath;
+		CreateFilePath(PayloadId, FilePath);
+
+		// If the file already exists we don't need to replace it, we will also do our own error logging.
+		if (!IFileManager::Get().Move(FilePath.ToString(), *TempFilePath, /*Replace*/ false, /*EvenIfReadOnly*/ false, /*Attributes*/ false, /*bDoNotRetryOrError*/ true))
 		{
-			UE_LOG(	LogVirtualization, Error, TEXT("[%s] Failed to move payload '%s' to it's final location '%s' due to system error: %s"),
+			// Store the error message in case we need to display it
+			TStringBuilder<MAX_SPRINTF> SystemErrorMsg;
+			Utils::GetFormattedSystemError(SystemErrorMsg);
+
+			IFileManager::Get().Delete(*TempFilePath, true, false, true); // Clean up the temp file if it is still around but do not failure cases to the user
+
+			// Check if another thread or process was writing out the payload at the same time, if so we 
+			// don't need to give an error message.
+			if (DoesPayloadExist(PayloadId))
+			{
+				UE_LOG(LogVirtualization, Verbose, TEXT("[%s] Already has a copy of the payload '%s'."), *GetDebugName(), *LexToString(PayloadId));
+				continue;
+			}
+			else
+			{
+				UE_LOG(LogVirtualization, Error, TEXT("[%s] Failed to move payload '%s' to it's final location '%s' due to system error: %s"),
 					*GetDebugName(),
-					*LexToString(Id),
+					*LexToString(PayloadId),
 					*FilePath,
 					SystemErrorMsg.ToString());
 
-			return EPushResult::Failed;
+				ErrorCount++;
+				Request.SetStatus(FPushRequest::EStatus::Failed);
+			}
 		}
 	}
 
-	return EPushResult::Success;
+	return ErrorCount == 0;
 }
 
 FCompressedBuffer FFileSystemBackend::PullData(const FIoHash& Id)
diff --git a/Engine/Source/Developer/Virtualization/Private/VirtualizationFileBackend.h b/Engine/Source/Developer/Virtualization/Private/VirtualizationFileBackend.h
--- a/Engine/Source/Developer/Virtualization/Private/VirtualizationFileBackend.h
+++ b/Engine/Source/Developer/Virtualization/Private/VirtualizationFileBackend.h
@@ -41,7 +41,7 @@ private:
 
 	virtual EConnectionStatus OnConnect() override;
 
-	virtual EPushResult PushData(const FIoHash& Id, const FCompressedBuffer& Payload, const FString& PackageContext) override;
+	virtual bool PushData(TArrayView<FPushRequest> Requests) override;
 
 	virtual FCompressedBuffer PullData(const FIoHash& Id) override;
 
diff --git a/Engine/Source/Developer/Virtualization/Private/VirtualizationManager.cpp b/Engine/Source/Developer/Virtualization/Private/VirtualizationManager.cpp
--- a/Engine/Source/Developer/Virtualization/Private/VirtualizationManager.cpp
+++ b/Engine/Source/Developer/Virtualization/Private/VirtualizationManager.cpp
@@ -422,12 +422,6 @@ bool FVirtualizationManager::AllowSubmitIfVirtualizationFailed() const
 	return bAllowSubmitIfVirtualizationFailed;
 }
 
-bool FVirtualizationManager::PushData(const FIoHash& Id, const FCompressedBuffer& Payload, EStorageType StorageType, const FString& Context)
-{
-	FPushRequest Request(Id, Payload, Context);
-	return FVirtualizationManager::PushData(MakeArrayView(&Request, 1), StorageType);
-}
-
 bool FVirtualizationManager::PushData(TArrayView<FPushRequest> Requests, EStorageType StorageType)
 {
 	TRACE_CPUPROFILER_EVENT_SCOPE(FVirtualizationManager::PushData);
@@ -1509,14 +1503,16 @@ void FVirtualizationManager::CachePayload(const FIoHash& Id, const FCompressedBu
 bool FVirtualizationManager::TryCacheDataToBackend(IVirtualizationBackend& Backend, const FIoHash& Id, const FCompressedBuffer& Payload)
 {
 	COOK_STAT(FCookStats::FScopedStatsCounter Timer(Profiling::GetCacheStats(Backend)));
-	const EPushResult Result = Backend.PushData(Id, Payload, FString());
-
-	if (Result == EPushResult::Success)
+	
+	if (Backend.PushData(Id, Payload, FString()))
 	{
 		COOK_STAT(Timer.AddHit(Payload.GetCompressedSize()));
+		return true;
+	}
+	else
+	{
+		return false;
 	}
-
-	return Result != EPushResult::Failed;
 }
 
 bool FVirtualizationManager::TryPushDataToBackend(IVirtualizationBackend& Backend, TArrayView<FPushRequest> Requests)
diff --git a/Engine/Source/Developer/Virtualization/Private/VirtualizationManager.h b/Engine/Source/Developer/Virtualization/Private/VirtualizationManager.h
--- a/Engine/Source/Developer/Virtualization/Private/VirtualizationManager.h
+++ b/Engine/Source/Developer/Virtualization/Private/VirtualizationManager.h
@@ -119,7 +119,6 @@ private:
 
 	virtual bool AllowSubmitIfVirtualizationFailed() const override;
 	
-	virtual bool PushData(const FIoHash& Id, const FCompressedBuffer& Payload, EStorageType StorageType, const FString& Context) override;
 	virtual bool PushData(TArrayView<FPushRequest> Requests, EStorageType StorageType) override;
 
 	virtual FCompressedBuffer PullData(const FIoHash& Id) override;
diff --git a/Engine/Source/Developer/Virtualization/Private/VirtualizationSourceControlBackend.cpp b/Engine/Source/Developer/Virtualization/Private/VirtualizationSourceControlBackend.cpp
--- a/Engine/Source/Developer/Virtualization/Private/VirtualizationSourceControlBackend.cpp
+++ b/Engine/Source/Developer/Virtualization/Private/VirtualizationSourceControlBackend.cpp
@@ -413,12 +413,6 @@ bool FSourceControlBackend::DoesPayloadExist(const FIoHash& Id)
 	}
 }
 
-EPushResult FSourceControlBackend::PushData(const FIoHash& Id, const FCompressedBuffer& Payload, const FString& Context)
-{
-	FPushRequest Request(Id, Payload, Context);
-	return FSourceControlBackend::PushData(MakeArrayView(&Request, 1)) ? EPushResult::Success : EPushResult::Failed;
-}
-
 bool FSourceControlBackend::PushData(TArrayView<FPushRequest> Requests)
 {
 	TRACE_CPUPROFILER_EVENT_SCOPE(FSourceControlBackend::PushData);
diff --git a/Engine/Source/Developer/Virtualization/Private/VirtualizationSourceControlBackend.h b/Engine/Source/Developer/Virtualization/Private/VirtualizationSourceControlBackend.h
--- a/Engine/Source/Developer/Virtualization/Private/VirtualizationSourceControlBackend.h
+++ b/Engine/Source/Developer/Virtualization/Private/VirtualizationSourceControlBackend.h
@@ -68,7 +68,6 @@ private:
 
 	virtual EConnectionStatus OnConnect() override;
 	
-	virtual EPushResult PushData(const FIoHash& Id, const FCompressedBuffer& Payload, const FString& Context) override;
 	virtual bool PushData(TArrayView<FPushRequest> Requests) override;
 
 	virtual FCompressedBuffer PullData(const FIoHash& Id) override;
diff --git a/Engine/Source/Developer/Virtualization/Public/IVirtualizationBackend.h b/Engine/Source/Developer/Virtualization/Public/IVirtualizationBackend.h
--- a/Engine/Source/Developer/Virtualization/Public/IVirtualizationBackend.h
+++ b/Engine/Source/Developer/Virtualization/Public/IVirtualizationBackend.h
@@ -14,17 +14,6 @@ struct FIoHash;
 namespace UE::Virtualization
 {
 
-/** Describes the result of a IVirtualizationBackend::Push operation */
-enum class EPushResult
-{
-	/** The push failed, the backend should print an error message to 'LogVirtualization'.*/
-	Failed = 0,
-	/** The payload already exists in the backend and does not need to be pushed. */
-	PayloadAlreadyExisted,
-	/** The payload was successfully pushed to the backend. */
-	Success
-};
-
 /**
  * The interface to derive from to create a new backend implementation.
  * 
@@ -120,36 +109,18 @@ public:
 	 * @param Payload	A potentially compressed buffer representing the payload
 	 * @return			The result of the push operation
 	 */
-	virtual EPushResult PushData(const FIoHash& Id, const FCompressedBuffer& Payload, const FString& PackageContext) = 0;
-
-	virtual bool PushData(TArrayView<FPushRequest> Requests)
+	bool PushData(const FIoHash& Id, const FCompressedBuffer& Payload, const FString& PackageContext)
 	{
-		// TODO: Improve the error codes in the future
-		for (FPushRequest& Request : Requests)
-		{
-			EPushResult Result = PushData(Request.GetIdentifier(), Request.GetPayload(), Request.GetContext());
-			switch (Result)
-			{
-			case EPushResult::Failed:
-				Request.SetStatus(FPushRequest::EStatus::Failed);
-				return false;
-
-			case EPushResult::PayloadAlreadyExisted:
-				// falls through
-			case EPushResult::Success:
-				Request.SetStatus(FPushRequest::EStatus::Success);
-				break;
-
-			default:
-				Request.SetStatus(FPushRequest::EStatus::Failed);
-				checkNoEntry();
-				break;
-			}
-		}
+		FPushRequest Request(Id, Payload, PackageContext);
 
-		return true;
+		return PushData(MakeArrayView(&Request, 1));
 	}
 
+	/**
+	 * @return True if all of the requests succeeded, false if one of more failed
+	 */
+	virtual bool PushData(TArrayView<FPushRequest> Requests) = 0;
+
 	/** 
 	 * The backend will attempt to retrieve the given payload by what ever method the backend uses.
 	 * NOTE: It is assumed that the virtualization manager will validate the returned payload to 
diff --git a/Engine/Source/Runtime/Core/Private/Virtualization/VirtualizationSystem.cpp b/Engine/Source/Runtime/Core/Private/Virtualization/VirtualizationSystem.cpp
--- a/Engine/Source/Runtime/Core/Private/Virtualization/VirtualizationSystem.cpp
+++ b/Engine/Source/Runtime/Core/Private/Virtualization/VirtualizationSystem.cpp
@@ -83,15 +83,6 @@ public:
 		return false;
 	}
 
-	virtual bool PushData(const FIoHash& Id, const FCompressedBuffer& Payload, EStorageType StorageType, const FString& Context) override
-	{
-#if UE_LOG_ON_NULLSYSTEM_USE
-		UE_LOG(LogVirtualization, UE_NULLSYSTEM_SEVERITY, TEXT("Cannot push payload '%s' as the virtualization system is disabled"), *LexToString(Id));
-#endif //UE_LOG_ON_NULLSYSTEM_USE
-		
-		return false;
-	}
-
 	virtual bool PushData(TArrayView<FPushRequest> Requests, EStorageType StorageType) override
 	{
 #if UE_LOG_ON_NULLSYSTEM_USE
diff --git a/Engine/Source/Runtime/Core/Public/Virtualization/VirtualizationSystem.h b/Engine/Source/Runtime/Core/Public/Virtualization/VirtualizationSystem.h
--- a/Engine/Source/Runtime/Core/Public/Virtualization/VirtualizationSystem.h
+++ b/Engine/Source/Runtime/Core/Public/Virtualization/VirtualizationSystem.h
@@ -393,36 +393,69 @@ public:
 	 *				process failed. If it returns false then the calling code should prevent the submit.
 	 */
 	virtual bool AllowSubmitIfVirtualizationFailed() const = 0;
-
-	/**
-	 * Push a payload to the virtualization backends.
-	 *
-	 * @param	Id				The identifier of the payload being pushed.
-	 * @param	Payload			The payload itself in FCompressedBuffer form, it is assumed that if the buffer is to 
-	 *							be compressed that it will have been done by the caller.
-	 * @param	StorageType		The type of storage to push the payload to, @See EStorageType for details.
-	 * @param	PackageContext	Context for the payload being submitted, typically the name from the UPackage that owns it.
-	 * 
-	 * @return	True if at least one backend now contains the payload, otherwise false.
-	 */
-	virtual bool PushData(const FIoHash& Id, const FCompressedBuffer& Payload, EStorageType StorageType, const FString& Context) = 0;
 	
 	/**
 	 * Push one or more payloads to a backend storage system. @See FPushRequest.
 	 * 
-	 * @param	Requests	A list of one or more payloads
+	 * @param	Requests	A collection of one or more payload push requests
 	 * @param	StorageType	The type of storage to push the payload to, @See EStorageType for details.
 	 * 
-	 * @return	When StorageType is Local, this method will return true assuming at least one backend
-	 *			managed to push all of the payloads. 
-	 *			When StorageType is Persistent, this method will only return true if ALL backends
-	 *			manage to push all of the payloads.
-	 * 			If this returns true then you can check the Status member of each request for more info 
-	 *			about each payloads push operation.
-	 *			If this returns false then you can assume that the payloads are not safely virtualized.
+	 * @return	When StorageType is EStorageType::Cache this will return true as long each payload
+	 *			ends up being stored in at least one of the cache backends.
+	 *			When StorageType is EStorageType::Persistent this will only return true if each payload
+	 *			endd up being stored in *all* of the backends.
+	 *			This is because the cache backends are not considered essential and it is not the end of
+	 *			the world if a payload is missing but the persistent backends must be reliable.
 	 */
 	virtual bool PushData(TArrayView<FPushRequest> Requests, EStorageType StorageType) = 0;
 
+	/**
+	 * Push a payload to the virtualization backends.
+	 *
+	 * @param	Request			A single push request
+	 * @param	PackageContext	Context for the payload being submitted, typically the name from the UPackage that owns it.
+	 *
+	 * @return	When StorageType is EStorageType::Cache this will return true as long each payload
+	 *			ends up being stored in at least one of the cache backends.
+	 *			When StorageType is EStorageType::Persistent this will only return true if each payload
+	 *			endd up being stored in *all* of the backends.
+	 *			This is because the cache backends are not considered essential and it is not the end of
+	 *			the world if a payload is missing but the persistent backends must be reliable.
+	 */
+	bool PushData(FPushRequest Request, EStorageType StorageType)
+	{
+		return PushData(MakeArrayView(&Request, 1), StorageType);
+	}
+
+	/**
+	 * Push a payload to the virtualization backends.
+	 *
+	 * @param	Id				The identifier of the payload being pushed.
+	 * @param	Payload			The payload itself in FCompressedBuffer form, it is assumed that if the buffer is to
+	 *							be compressed that it will have been done by the caller.
+	 * @param	Context			Context for the payload being submitted, typically the name from the UPackage that owns it.
+	 * @param	StorageType		The type of storage to push the payload to, @See EStorageType for details.
+	 *
+	 * @return	When StorageType is EStorageType::Cache this will return true as long each payload
+	 *			ends up being stored in at least one of the cache backends.
+	 *			When StorageType is EStorageType::Persistent this will only return true if each payload
+	 *			endd up being stored in *all* of the backends.
+	 *			This is because the cache backends are not considered essential and it is not the end of
+	 *			the world if a payload is missing but the persistent backends must be reliable.
+	 */
+	bool PushData(const FIoHash& Id, const FCompressedBuffer& Payload, const FString& Context, EStorageType StorageType)
+	{
+		FPushRequest Request(Id, Payload, Context);
+		return PushData(MakeArrayView(&Request, 1), StorageType);
+	}
+
+	UE_DEPRECATED(5.1, "Use the overload of ::PushData(FIoHash, FCompressedBuffer, FString, EStorageType)")
+	bool PushData(const FIoHash& Id, const FCompressedBuffer& Payload, EStorageType StorageType, const FString& Context)
+	{
+		FPushRequest Request(Id, Payload, Context);
+		return PushData(MakeArrayView(&Request, 1), StorageType);
+	}
+
 	/**
 	 * Pull a payload from the virtualization backends.
 	 *
diff --git a/Engine/Source/Runtime/CoreUObject/Private/Serialization/EditorBulkData.cpp b/Engine/Source/Runtime/CoreUObject/Private/Serialization/EditorBulkData.cpp
--- a/Engine/Source/Runtime/CoreUObject/Private/Serialization/EditorBulkData.cpp
+++ b/Engine/Source/Runtime/CoreUObject/Private/Serialization/EditorBulkData.cpp
@@ -1445,7 +1445,7 @@ void FEditorBulkData::PushData(const FPackagePath& InPackagePath)
 		RecompressForSerialization(PayloadToPush, Flags);
 
 		// TODO: We could make the storage type a config option?
-		if (VirtualizationSystem.PushData(PayloadContentId, PayloadToPush, UE::Virtualization::EStorageType::Cache, InPackagePath.GetPackageName()))
+		if (VirtualizationSystem.PushData(PayloadContentId, PayloadToPush, InPackagePath.GetPackageName(), UE::Virtualization::EStorageType::Cache))
 		{
 			EnumAddFlags(Flags, EFlags::IsVirtualized);
 			EnumRemoveFlags(Flags, EFlags::ReferencesLegacyFile | EFlags::ReferencesWorkspaceDomain | EFlags::LegacyFileIsCompressed);
