The server address used by the source control VA backend can now be set when describing the backend graph in the config file.

#rnx

- This is only useful if the entire team uses the same address for the perforce server
- When set, this will override the environment settings for P4PORT as well as override what ever is saved to SourceControlSettings.ini

diff --git a/Engine/Source/Developer/Virtualization/Private/VirtualizationSourceControlBackend.cpp b/Engine/Source/Developer/Virtualization/Private/VirtualizationSourceControlBackend.cpp
--- a/Engine/Source/Developer/Virtualization/Private/VirtualizationSourceControlBackend.cpp
+++ b/Engine/Source/Developer/Virtualization/Private/VirtualizationSourceControlBackend.cpp
@@ -261,6 +261,10 @@ IVirtualizationBackend::EConnectionStatus FSourceControlBackend::OnConnect()
 
 	// We do not want the connection to have a client workspace so explicitly set it to empty
 	FSourceControlInitSettings SCCSettings(FSourceControlInitSettings::EBehavior::OverrideExisting);
+	if (!ServerAddress.IsEmpty())
+	{
+		SCCSettings.AddSetting(TEXT("P4Port"), ServerAddress);
+	}
 	SCCSettings.AddSetting(TEXT("P4Client"), TEXT(""));
 
 	SCCProvider = ISourceControlModule::Get().CreateProvider(FName("Perforce"), TEXT("Virtualization"), SCCSettings);
@@ -810,6 +814,19 @@ bool FSourceControlBackend::TryApplySettingsFromConfigFiles(const FString& Confi
 
 	// Now parse the optional config values
 
+	{
+		FParse::Value(*ConfigEntry, TEXT("Server="), ServerAddress);
+
+		if (!ClientStream.IsEmpty())
+		{
+			UE_LOG(LogVirtualization, Log, TEXT("[%s] Using server address: '%s'"), *GetDebugName(), *ServerAddress);
+		}
+		else
+		{
+			UE_LOG(LogVirtualization, Log, TEXT("[%s] Will connect to the default server address"), *GetDebugName());
+		}
+	}
+
 	// If the depot root is within a perforce stream then we must specify which stream. This may be a virtual stream with a custom view.
 	{
 		FParse::Value(*ConfigEntry, TEXT("ClientStream="), ClientStream);
diff --git a/Engine/Source/Developer/Virtualization/Private/VirtualizationSourceControlBackend.h b/Engine/Source/Developer/Virtualization/Private/VirtualizationSourceControlBackend.h
--- a/Engine/Source/Developer/Virtualization/Private/VirtualizationSourceControlBackend.h
+++ b/Engine/Source/Developer/Virtualization/Private/VirtualizationSourceControlBackend.h
@@ -26,8 +26,10 @@ class FSemaphore;
  * depot where the payload files are being stored.
  * 
  * Optional Values:
+ * Server [string]:				When set the backend will use this server address to connect to. When not
+ *								set the backend will use the environment defaults. [Default=""]
  * ClientStream [string]:		Used when the payloads are stored in a stream based depot. It should contain
- *								the stream name to use when creating a workspace for payload submission.
+ *								the stream name to use when creating a workspace for payload submission. [Default=""]
  * UsePartitionedClient [bool]:	When true the temporary workspace client created to submit payloads 
  *								from will be created as a partitioned workspace which is less overhead
  *								on the source control server. If your server does not support this then
@@ -96,6 +98,9 @@ private:
 	/** The root where the virtualized payloads are stored in source control */
 	FString DepotRoot;
 
+	/** Address of the server to connect to */
+	FString ServerAddress;
+
 	/** The stream containing the DepotRoot where the virtualized payloads are stored in source control */
 	FString ClientStream;
 
