Added cvar 'VA.AllowPkgVirtualization' which can dynamically enable/disable the virtualization process in the editor without restarting.

#rnx

- This replaces the old cvar 'VA.DisableVirtualization' which is being removed.
- We will log each time the value is switched to aid in future debugging.
- The name needed to be changed as we needed to invert how it worked, also "DiableVirtualization" was a bit misleading and made people think it turned the entire system off, not just the process of virtualizing packages on submit.
-- The member bEnablePayloadVirtualization was also renamed to bAllowPackageVirtualization to make its purpose clearer.
- The cmdline version was changed from 'VA-DisableVirtualization' to 'VA-SkipPkgVirtualization' for similar reasons.

diff --git a/Engine/Source/Developer/Virtualization/Private/VirtualizationManager.cpp b/Engine/Source/Developer/Virtualization/Private/VirtualizationManager.cpp
--- a/Engine/Source/Developer/Virtualization/Private/VirtualizationManager.cpp
+++ b/Engine/Source/Developer/Virtualization/Private/VirtualizationManager.cpp
@@ -37,12 +37,6 @@ static TAutoConsoleVariable<bool> CVarLazyInitConnections(
 	false,
 	TEXT("When true the VA backends will defer creating their connections until first use"));
 
-// TODO: Move to RegisterConsoleCommands
-static TAutoConsoleVariable<bool> CVarDisableVirtualization(
-	TEXT("VA.DisableVirtualization"),
-	false,
-	TEXT("When true submitting packages in the editor will no longer trigger the virtualization process (only applied when the system is initialized)"));
-
 /** Utility struct, similar to FScopeLock but allows the lock to be enabled/disabled more easily */
 struct FConditionalScopeLock
 {
@@ -289,7 +283,7 @@ namespace Profiling
 } //namespace Profiling
 
 FVirtualizationManager::FVirtualizationManager()
-	: bEnablePayloadVirtualization(true)
+	: bAllowPackageVirtualization(true)
 	, bEnableCacheAfterPull(true)
 	, MinPayloadLength(0)
 	, BackendGraphName(TEXT("ContentVirtualizationBackendGraph_None"))
@@ -362,7 +356,7 @@ bool FVirtualizationManager::IsEnabled() const
 
 bool FVirtualizationManager::IsPushingEnabled(EStorageType StorageType) const
 {
-	if (!bEnablePayloadVirtualization)
+	if (!bAllowPackageVirtualization)
 	{
 		return false;
 	}
@@ -482,7 +476,7 @@ bool FVirtualizationManager::PushData(TArrayView<FPushRequest> Requests, EStorag
 	}
 
 	// Early out if there are no backends
-	if (!IsEnabled() || bEnablePayloadVirtualization == false)
+	if (!IsEnabled() || !bAllowPackageVirtualization)
 	{
 		return false;
 	}
@@ -805,8 +799,8 @@ void FVirtualizationManager::ApplySettingsFromConfigFiles(const FConfigFile& Con
 
 		if (bLoadedFromFile)
 		{
-			bEnablePayloadVirtualization = bEnablePayloadVirtualizationFromIni;
-			UE_LOG(LogVirtualization, Display, TEXT("\tEnablePayloadVirtualization : %s"), bEnablePayloadVirtualization ? TEXT("true") : TEXT("false"));
+			bAllowPackageVirtualization = bEnablePayloadVirtualizationFromIni;
+			UE_LOG(LogVirtualization, Display, TEXT("\tEnablePayloadVirtualization : %s"), bAllowPackageVirtualization ? TEXT("true") : TEXT("false"));
 		}
 		else
 		{
@@ -996,10 +990,10 @@ void FVirtualizationManager::ApplySettingsFromFromCmdline()
 		UE_LOG(LogVirtualization, Display, TEXT("Cmdline has set the virtualization system backends to lazy init their connections"));
 	}
 
-	if (bEnablePayloadVirtualization && FParse::Param(FCommandLine::Get(), TEXT("VA-DisableVirtualization")))
+	if (bAllowPackageVirtualization && FParse::Param(FCommandLine::Get(), TEXT("VA-SkipPkgVirtualization")))
 	{
-		bEnablePayloadVirtualization = false;
-		UE_LOG(LogVirtualization, Warning, TEXT("Cmdline has disabled the virtualization process"));
+		bAllowPackageVirtualization = false;
+		UE_LOG(LogVirtualization, Warning, TEXT("The virtualization process has been disabled via the command line"));
 	}
 }
 
@@ -1010,12 +1004,6 @@ void FVirtualizationManager::ApplySettingsFromCVar()
 		bLazyInitConnections = true;
 		UE_LOG(LogVirtualization, Display, TEXT("CVar has set the virtualization system backends to lazy init their connections"));
 	}
-
-	if (bEnablePayloadVirtualization && CVarDisableVirtualization.GetValueOnAnyThread())
-	{
-		bEnablePayloadVirtualization = false;
-		UE_LOG(LogVirtualization, Display, TEXT("CVar has disabled the virtualization process"));
-	}
 }
 
 void FVirtualizationManager::ApplyDebugSettingsFromFromCmdline()
@@ -1062,6 +1050,38 @@ void FVirtualizationManager::ApplyDebugSettingsFromFromCmdline()
 
 void FVirtualizationManager::RegisterConsoleCommands()
 {
+	{
+		bool bOriginalValue = bAllowPackageVirtualization;
+
+		IConsoleVariable* Handle = IConsoleManager::Get().RegisterConsoleVariableRef(
+			TEXT("VA.AllowPkgVirtualization"),
+			bAllowPackageVirtualization,
+			TEXT("When true submitting packages in the editor will no longer trigger the virtualization process")
+		);
+
+		auto Callback = [](IConsoleVariable* CVar)
+			{
+				if (CVar->GetBool())
+				{
+					UE_LOG(LogVirtualization, Display, TEXT("The virtualization process has been enabled via the cvar 'VA.SkipPkgVirtualization'"));
+				}
+				else
+				{
+					UE_LOG(LogVirtualization, Display, TEXT("The virtualization process has been disabled via the cvar 'VA.SkipPkgVirtualization'"));
+				}
+			};
+		
+		// Log the change if the cvar was modified on the commandline
+		if (bOriginalValue != bAllowPackageVirtualization)
+		{
+			Callback(Handle);
+		}
+
+		Handle->OnChangedDelegate().AddLambda(MoveTemp(Callback));
+
+		DebugValues.ConsoleObjects.Add(Handle);
+	}
+
 	DebugValues.ConsoleObjects.Add(IConsoleManager::Get().RegisterConsoleCommand(
 		TEXT("VA.MissBackends"),
 		TEXT("A debug commnad which can be used to disable payload pulling on one or more backends"),
diff --git a/Engine/Source/Developer/Virtualization/Private/VirtualizationManager.h b/Engine/Source/Developer/Virtualization/Private/VirtualizationManager.h
--- a/Engine/Source/Developer/Virtualization/Private/VirtualizationManager.h
+++ b/Engine/Source/Developer/Virtualization/Private/VirtualizationManager.h
@@ -206,8 +206,8 @@ private:
 private:
 	// The following members are set from the config file
 
-	/** Are payloads allowed to be virtualized. Defaults to true. */
-	bool bEnablePayloadVirtualization;
+	/** Are packages allowed to be virtualized when submitted to source control. Defaults to true. */
+	bool bAllowPackageVirtualization;
 
 	/** Should payloads be cached locally after being pulled from persistent storage? Defaults to true. */
 	bool bEnableCacheAfterPull;
