The VA changelist stamp should only be applied if the VA process actually ran and was not disabled.

#rnx

diff --git a/Engine/Source/Developer/Virtualization/Private/PackageVirtualizationProcess.cpp b/Engine/Source/Developer/Virtualization/Private/PackageVirtualizationProcess.cpp
--- a/Engine/Source/Developer/Virtualization/Private/PackageVirtualizationProcess.cpp
+++ b/Engine/Source/Developer/Virtualization/Private/PackageVirtualizationProcess.cpp
@@ -160,19 +160,6 @@ void VirtualizePackages(const TArray<FString>& FilesToSubmit, TArray<FText>& Out
 
 	IVirtualizationSystem& System = IVirtualizationSystem::Get();
 
-	// TODO: We could check to see if the package is virtualized even if it is disabled for the project
-	// as a safety feature?
-	if (!System.IsEnabled())
-	{
-		return;
-	}
-
-	if (!System.IsPushingEnabled(EStorageType::Persistent))
-	{
-		UE_LOG(LogVirtualization, Verbose, TEXT("Pushing to persistent backend storage is disabled"));
-		return;
-	}
-
 	const double StartTime = FPlatformTime::Seconds();
 
 	FScopedSlowTask Progress(5.0f, LOCTEXT("Virtualization_Task", "Virtualizing Assets..."));
diff --git a/Engine/Source/Developer/Virtualization/Private/VirtualizationManager.cpp b/Engine/Source/Developer/Virtualization/Private/VirtualizationManager.cpp
--- a/Engine/Source/Developer/Virtualization/Private/VirtualizationManager.cpp
+++ b/Engine/Source/Developer/Virtualization/Private/VirtualizationManager.cpp
@@ -671,15 +671,15 @@ EVirtualizationResult FVirtualizationManager::TryVirtualizePackages(const TArray
 	OutDescriptionTags.Reset();
 	OutErrors.Reset();
 
-	UE::Virtualization::VirtualizePackages(FilesToVirtualize, OutErrors);
-
-	// If we had no new errors add the validation tag to indicate that the packages are safe for submission. 
-	// TODO: Currently this is a simple tag to make it easier for us to track which assets were submitted via the
-	// virtualization process in a test project. This should be expanded when we add proper p4 server triggers.
-	if (OutErrors.IsEmpty() && !VirtualizationProcessTag.IsEmpty())
+	if (IsEnabled() && IsPushingEnabled(EStorageType::Persistent))
 	{
-		FText Tag = FText::FromString(VirtualizationProcessTag);
-		OutDescriptionTags.Add(Tag);
+		UE::Virtualization::VirtualizePackages(FilesToVirtualize, OutErrors);
+
+		if (OutErrors.IsEmpty() && !VirtualizationProcessTag.IsEmpty())
+		{
+			FText Tag = FText::FromString(VirtualizationProcessTag);
+			OutDescriptionTags.Add(Tag);
+		}
 	}
 
 	return OutErrors.IsEmpty() ? EVirtualizationResult::Success : EVirtualizationResult::Failed;
