Setting the cvar 'VA.AllowPkgVirtualization' as part of a device profile will no longer log warnings about the cvar being unregistered.

#rnx

- Everything was working but when a cvar is set via a device profile we expect the cvar to be registered sooner, hence the warning.
- Added a new define UE_USE_GLOBAL_CVAR to FVirtualizationManager.cpp, when set we use a static global cvar for 'VA.AllowPkgVirtualization' rather then creating and registering one when the system is mounted.
- Since the cvar can now outlive the manager we also needed to add a way to remove the bound lambda from the cvar when the manager is destroyed.

diff --git a/Engine/Source/Developer/Virtualization/Private/VirtualizationManager.cpp b/Engine/Source/Developer/Virtualization/Private/VirtualizationManager.cpp
--- a/Engine/Source/Developer/Virtualization/Private/VirtualizationManager.cpp
+++ b/Engine/Source/Developer/Virtualization/Private/VirtualizationManager.cpp
@@ -37,6 +37,19 @@ static TAutoConsoleVariable<bool> CVarLazyInitConnections(
 	false,
 	TEXT("When true the VA backends will defer creating their connections until first use"));
 
+//
+#define UE_USE_GLOBAL_CVAR 1
+
+#if UE_USE_GLOBAL_CVAR
+
+static FAutoConsoleVariable CVarAllowPkgVirtualization(
+	TEXT("VA.AllowPkgVirtualization"),
+	true,
+	TEXT("When true submitting packages in the editor will no longer trigger the virtualization process")
+);
+
+#endif // UE_USE_GLOBAL_CVAR
+
 /** Utility struct, similar to FScopeLock but allows the lock to be enabled/disabled more easily */
 struct FConditionalScopeLock
 {
@@ -300,6 +313,14 @@ FVirtualizationManager::FVirtualizationManager()
 
 FVirtualizationManager::~FVirtualizationManager()
 {
+	for (const TPair<IConsoleVariable*, FDelegateHandle>& KV : DebugValues.ConsoleDelegateHandles)
+	{
+		IConsoleVariable* ConsoleVariable = KV.Key;
+		const FDelegateHandle& Handle = KV.Value;
+
+		ConsoleVariable->OnChangedDelegate().Remove(Handle);
+	}
+
 	for (IConsoleObject* ConsoleObject : DebugValues.ConsoleObjects)
 	{
 		IConsoleManager::Get().UnregisterConsoleObject(ConsoleObject);
@@ -1051,6 +1072,39 @@ void FVirtualizationManager::ApplyDebugSettingsFromFromCmdline()
 void FVirtualizationManager::RegisterConsoleCommands()
 {
 	{
+#if UE_USE_GLOBAL_CVAR
+		IConsoleVariable* Handle = IConsoleManager::Get().FindConsoleVariable(TEXT("VA.AllowPkgVirtualization"));
+
+		if (Handle != nullptr)
+		{
+			auto Callback = [this](IConsoleVariable* CVar)
+			{
+				this->bAllowPackageVirtualization = CVar->GetBool();
+
+				if (this->bAllowPackageVirtualization)
+				{
+					UE_LOG(LogVirtualization, Display, TEXT("The virtualization process has been enabled via the cvar 'VA.SkipPkgVirtualization'"));
+				}
+				else
+				{
+					UE_LOG(LogVirtualization, Display, TEXT("The virtualization process has been disabled via the cvar 'VA.SkipPkgVirtualization'"));
+				}
+			};
+
+			FDelegateHandle CallbackHandle = Handle->OnChangedDelegate().AddLambda(MoveTemp(Callback));
+
+			if (bAllowPackageVirtualization != Handle->GetBool())
+			{
+				Callback(Handle);
+			}
+
+			DebugValues.ConsoleDelegateHandles.Add({ Handle, CallbackHandle });
+		}
+		else
+		{
+			UE_LOG(LogVirtualization, Warning, TEXT("CVar VA.AllowPkgVirtualization could not be found and will not function"));
+		}
+#else
 		bool bOriginalValue = bAllowPackageVirtualization;
 
 		IConsoleVariable* Handle = IConsoleManager::Get().RegisterConsoleVariableRef(
@@ -1080,6 +1134,7 @@ void FVirtualizationManager::RegisterConsoleCommands()
 		Handle->OnChangedDelegate().AddLambda(MoveTemp(Callback));
 
 		DebugValues.ConsoleObjects.Add(Handle);
+#endif // UE_USE_GLOBAL_CVAR
 	}
 
 	DebugValues.ConsoleObjects.Add(IConsoleManager::Get().RegisterConsoleCommand(
diff --git a/Engine/Source/Developer/Virtualization/Private/VirtualizationManager.h b/Engine/Source/Developer/Virtualization/Private/VirtualizationManager.h
--- a/Engine/Source/Developer/Virtualization/Private/VirtualizationManager.h
+++ b/Engine/Source/Developer/Virtualization/Private/VirtualizationManager.h
@@ -275,6 +275,13 @@ private:
 		/** All of the console commands/variables that we register, so they can be unregistered when the manager is destroyed */
 		TArray<IConsoleObject*> ConsoleObjects;
 
+		/** 
+		 * Contains all of the delegate handles that we have bound to IConsoleVariable::OnChangedDelegate and need to be removed
+		 * before it is safe to destroy the manager. Most likely due to having bound a lambda to the delegate that captured the
+		 * FVirtualizationManager this pointer.
+		 */
+		TArray<TPair<IConsoleVariable*, FDelegateHandle>> ConsoleDelegateHandles;
+
 		/** The critical section used to force single threaded access if bForceSingleThreaded is true */
 		FCriticalSection ForceSingleThreadedCS;
 
