We now also log the VA payload count along side the other statistics.

#rnx

- In addition to adding the hit counter (which can be very useful when debugging) the stats logging code has been factored out to a single lambda tht is called for each set of stats to make it easier to update in future.

diff --git a/Engine/Source/Developer/Virtualization/Private/VirtualizationManager.cpp b/Engine/Source/Developer/Virtualization/Private/VirtualizationManager.cpp
--- a/Engine/Source/Developer/Virtualization/Private/VirtualizationManager.cpp
+++ b/Engine/Source/Developer/Virtualization/Private/VirtualizationManager.cpp
@@ -221,7 +221,7 @@ namespace Profiling
 	{
 		UE_LOG(LogVirtualization, Display, TEXT(""));
 		UE_LOG(LogVirtualization, Display, TEXT("Virtualization ProfileData"));
-		UE_LOG(LogVirtualization, Display, TEXT("======================================================================================="));
+		UE_LOG(LogVirtualization, Display, TEXT("=================================================================================================="));
 
 		if (!HasProfilingData())
 		{
@@ -229,68 +229,35 @@ namespace Profiling
 			return; // Early out if we have no data
 		}
 
-		if (HasProfilingData(CacheStats))
-		{
-			UE_LOG(LogVirtualization, Display, TEXT("%-40s|%17s|%12s|%14s|"), TEXT("Caching Data"), TEXT("TotalSize (MB)"), TEXT("TotalTime(s)"), TEXT("DataRate(MB/S)"));
-			UE_LOG(LogVirtualization, Display, TEXT("----------------------------------------|-----------------|------------|--------------|"));
-
-			for (const auto& Iterator : CacheStats)
+		auto PrintStats = [](const TCHAR* Name, const TMap<FString, FCookStats::CallStats>& Stats)
 			{
-				const double Time = Iterator.Value.GetAccumulatedValueAnyThread(FCookStats::CallStats::EHitOrMiss::Hit, FCookStats::CallStats::EStatType::Cycles) * FPlatformTime::GetSecondsPerCycle();
-				const double DataSizeMB = Iterator.Value.GetAccumulatedValueAnyThread(FCookStats::CallStats::EHitOrMiss::Hit, FCookStats::CallStats::EStatType::Bytes) / (1024.0f * 1024.0f);
-				const double MBps = Time != 0.0 ? (DataSizeMB / Time) : 0.0;
-
-				UE_LOG(LogVirtualization, Display, TEXT("%-40.40s|%17.1f|%12.3f|%14.3f|"),
-					*Iterator.Key,
-					DataSizeMB,
-					Time,
-					MBps);
-			}
-
-			UE_LOG(LogVirtualization, Display, TEXT("======================================================================================="));
-		}
-
-		if (HasProfilingData(PushStats))
-		{
-			UE_LOG(LogVirtualization, Display, TEXT("%-40s|%17s|%12s|%14s|"), TEXT("Pushing Data"), TEXT("TotalSize (MB)"), TEXT("TotalTime(s)"), TEXT("DataRate(MB/S)"));
-			UE_LOG(LogVirtualization, Display, TEXT("----------------------------------------|-----------------|------------|--------------|"));
-
-			for (const auto& Iterator : PushStats)
-			{
-				const double Time = Iterator.Value.GetAccumulatedValueAnyThread(FCookStats::CallStats::EHitOrMiss::Hit, FCookStats::CallStats::EStatType::Cycles) * FPlatformTime::GetSecondsPerCycle();
-				const double DataSizeMB = Iterator.Value.GetAccumulatedValueAnyThread(FCookStats::CallStats::EHitOrMiss::Hit, FCookStats::CallStats::EStatType::Bytes) / (1024.0f * 1024.0f);
-				const double MBps = Time != 0.0 ? (DataSizeMB / Time) : 0.0;
-
-				UE_LOG(LogVirtualization, Display, TEXT("%-40.40s|%17.1f|%12.3f|%14.3f|"),
-					*Iterator.Key,
-					DataSizeMB,
-					Time,
-					MBps);
-			}
-
-			UE_LOG(LogVirtualization, Display, TEXT("======================================================================================="));
-		}
+				if (HasProfilingData(Stats))
+				{
+					UE_LOG(LogVirtualization, Display, TEXT("%-40s|%10s|%17s|%12s|%14s|"), Name, TEXT("TotalCount"), TEXT("TotalSize (MB)"), TEXT("TotalTime(s)"), TEXT("DataRate(MB/S)"));
+					UE_LOG(LogVirtualization, Display, TEXT("----------------------------------------|----------|-----------------|------------|--------------|"));
 
-		if (HasProfilingData(PullStats))
-		{
-			UE_LOG(LogVirtualization, Display, TEXT("%-40s|%17s|%12s|%14s|"), TEXT("Pulling Data"), TEXT("TotalSize (MB)"), TEXT("TotalTime(s)"), TEXT("DataRate(MB/S)"));
-			UE_LOG(LogVirtualization, Display, TEXT("----------------------------------------|-----------------|------------|--------------|"));
+					for (const auto& Iterator : Stats)
+					{
+						const int64 Count = Iterator.Value.GetAccumulatedValueAnyThread(FCookStats::CallStats::EHitOrMiss::Hit, FCookStats::CallStats::EStatType::Counter);
+						const double Time = Iterator.Value.GetAccumulatedValueAnyThread(FCookStats::CallStats::EHitOrMiss::Hit, FCookStats::CallStats::EStatType::Cycles) * FPlatformTime::GetSecondsPerCycle();
+						const double DataSizeMB = Iterator.Value.GetAccumulatedValueAnyThread(FCookStats::CallStats::EHitOrMiss::Hit, FCookStats::CallStats::EStatType::Bytes) / (1024.0f * 1024.0f);
+						const double MBps = Time != 0.0 ? (DataSizeMB / Time) : 0.0;
+
+						UE_LOG(LogVirtualization, Display, TEXT("%-40.40s|%10lld|%17.1f|%12.3f|%14.3f|"),
+							*Iterator.Key,
+							Count,
+							DataSizeMB,
+							Time,
+							MBps);
+					}
 
-			for (const auto& Iterator : PullStats)
-			{
-				const double Time = Iterator.Value.GetAccumulatedValueAnyThread(FCookStats::CallStats::EHitOrMiss::Hit, FCookStats::CallStats::EStatType::Cycles) * FPlatformTime::GetSecondsPerCycle();
-				const double DataSizeMB = Iterator.Value.GetAccumulatedValueAnyThread(FCookStats::CallStats::EHitOrMiss::Hit, FCookStats::CallStats::EStatType::Bytes) / (1024.0f * 1024.0f);
-				const double MBps = Time != 0.0 ? (DataSizeMB / Time) : 0.0;
-
-				UE_LOG(LogVirtualization, Display, TEXT("%-40.40s|%17.1f|%12.3f|%14.3f|"),
-					*Iterator.Key,
-					DataSizeMB,
-					Time,
-					MBps);
-			}
+					UE_LOG(LogVirtualization, Display, TEXT("=================================================================================================="));
+				}
+			};
 
-			UE_LOG(LogVirtualization, Display, TEXT("======================================================================================="));
-		}
+		PrintStats(TEXT("Caching Data"), CacheStats);
+		PrintStats(TEXT("Pushing Data"), PushStats);
+		PrintStats(TEXT("Pulling Data"), PullStats);
 	}
 #endif // ENABLE_COOK_STATS
 } //namespace Profiling
