Add a VA rehydration code path that rehydrates a package to a memory buffer rather than replacing the original package on disk.

#rnx

- Added a new FArchive 'FFixedBufferWriterArchive' that will let me serialize to an already allocated, fixed sized memory buffer.
- Moved the code for opening and reading packages to shared utility functions so that it is easier to keep the error messages consistent.
- Moved common code to a new function ::TryRehydrateBuilder so that both rehydration methods can share it.
- Note that technically we are changing the API since now we accept a TConstArrayView for the package file paths for both the virtualization and rehydration paths, but existing code will be compatible with this change.

diff --git a/Engine/Source/Developer/Virtualization/Private/PackageRehydrationProcess.cpp b/Engine/Source/Developer/Virtualization/Private/PackageRehydrationProcess.cpp
--- a/Engine/Source/Developer/Virtualization/Private/PackageRehydrationProcess.cpp
+++ b/Engine/Source/Developer/Virtualization/Private/PackageRehydrationProcess.cpp
@@ -8,6 +8,7 @@
 #include "Misc/PackagePath.h"
 #include "Misc/ScopedSlowTask.h"
 #include "PackageUtils.h"
+#include "Serialization/MemoryArchive.h"
 #include "UObject/Linker.h"
 #include "UObject/Package.h"
 #include "UObject/PackageTrailer.h"
@@ -19,85 +20,178 @@
 namespace UE::Virtualization
 {
 
-void RehydratePackages(const TArray<FString>& Packages, TArray<FText>& OutErrors)
+/** 
+ * Archive allowing serialization to a fixed sized memory buffer. Exceeding the
+ * length of the buffer will cause the archive to log a warning and set its
+ * error state to true.
+ */
+class FFixedBufferWriterArchive : public FMemoryArchive
 {
-	TRACE_CPUPROFILER_EVENT_SCOPE(UE::Virtualization::RehydratePackages);
+public:
+	FFixedBufferWriterArchive() = delete;
 
-	IVirtualizationSystem& System = IVirtualizationSystem::Get();
+	/** 
+	 * @param InBuffer Pointer to the buffer to write to
+	 * @param InLength Total length (in bytes) that can be written to the buffer
+	 */
+	FFixedBufferWriterArchive(uint8* InBuffer, int64 InLength)
+		: Buffer(InBuffer)
+		, Length(InLength)
+	{
+		SetIsSaving(true);
+	}
 
-	if (!System.IsEnabled())
+	virtual ~FFixedBufferWriterArchive() = default;
+	
+	/** Returns how much space remains in the internal buffer. */
+	int64 GetRemainingSpace() const
 	{
-		return;
+		return Length - Offset;
 	}
 
-	FScopedSlowTask Progress(1.0f, LOCTEXT("VAHydration_Task", "Re-hydrating Assets..."));
-	Progress.MakeDialog();
+private:
 
-	for (const FString& FilePath : Packages)
+	virtual void Serialize(void* Data, int64 Num) override
 	{
-		if (!FPackageName::IsPackageFilename(FilePath))
+		if (Offset + Num <= Length)
 		{
-			continue; // Only rehydrate valid packages
+			FMemory::Memcpy(&Buffer[Offset], Data, Num);
+			Offset += Num;
 		}
-			
-		FPackageTrailer Trailer;
-		if (!FPackageTrailer::TryLoadFromFile(FilePath, Trailer))
+		else
 		{
-			continue; // Only rehydrate packages with package trailers
+			UE_LOG(LogSerialization, Error, TEXT("Attempting to write %lld bytes to a FFixedBufferWriterArchive with only %lld bytes of space remaining"), Num, GetRemainingSpace());
+			SetError();
 		}
+	}
 
-		TArray<FIoHash> VirtualizedPayloads = Trailer.GetPayloads(EPayloadStorageType::Virtualized);
-		if (VirtualizedPayloads.IsEmpty())
-		{
-			continue; // If the package has no virtualized payloads then we can skip the rest
-		}
+	uint8* Buffer = nullptr;
+	int64 Length = INDEX_NONE;
+};
 
-		TUniquePtr<FArchive> PackageAr(IFileManager::Get().CreateFileReader(*FilePath));
-		if (!PackageAr.IsValid())
-		{
-			FText Message = FText::Format(LOCTEXT("VAHydration_ReadFailed", "Unable to open the package '{0}' for reading"),
-				FText::FromString(FilePath));
-			OutErrors.Add(Message);
+/** Utility for opening a package for reading, including localized error handling */
+TUniquePtr<FArchive> OpenFileForReading(const FString& FilePath, TArray<FText>& OutErrors)
+{
+	TUniquePtr<FArchive> FileHandle(IFileManager::Get().CreateFileReader(*FilePath));
+	if (!FileHandle)
+	{
+		FText Message = FText::Format(LOCTEXT("VAHydration_OpenFailed", "Unable to open the file '{0}' for reading"),
+			FText::FromString(FilePath));
+
+		OutErrors.Add(MoveTemp(Message));
+	}
+
+	return FileHandle;
+}
+
+/** Utility for serializing data from a package, including localized error handling */
+bool ReadData(TUniquePtr<FArchive>& Ar, void* DstBuffer, int64 BytesToRead, const FString& FilePath, TArray<FText>& OutErrors)
+{
+	check(Ar.IsValid());
+
+	Ar->Serialize(DstBuffer, BytesToRead);
+
+	if (!Ar->IsError())
+	{
+		return true;
+	}
+	else
+	{
+		FText Message = FText::Format(LOCTEXT("VAHydration_ReadFailed", "Failed to read {0} bytes from file {1}"),
+			BytesToRead,
+			FText::FromString(FilePath));
+
+		OutErrors.Add(MoveTemp(Message));
+		return false;
+	}
+}
+
+/** Shared rehydration code */
+bool TryRehydrateBuilder(const FString& FilePath, FPackageTrailer& OutTrailer, FPackageTrailerBuilder& OutBuilder, TArray<FText>& OutErrors)
+{
+	if (!FPackageName::IsPackageFilename(FilePath))
+	{
+		return false; // Only rehydrate valid packages
+	}
 
-			return;
+	if (!FPackageTrailer::TryLoadFromFile(FilePath, OutTrailer))
+	{
+		return false; // Only rehydrate packages with package trailers
+	}
+
+	TArray<FIoHash> VirtualizedPayloads = OutTrailer.GetPayloads(EPayloadStorageType::Virtualized);
+	if (VirtualizedPayloads.IsEmpty())
+	{
+		return false; // If the package has no virtualized payloads then we can skip the rest
+	}
+
+	{
+		TUniquePtr<FArchive> FileAr = OpenFileForReading(FilePath, OutErrors);
+		if (!FileAr)
+		{
+			return false;
 		}
 
-		FPackageTrailerBuilder Builder = FPackageTrailerBuilder::CreateFromTrailer(Trailer, *PackageAr, FilePath);
-		PackageAr.Reset();
+		OutBuilder = FPackageTrailerBuilder::CreateFromTrailer(OutTrailer, *FileAr, FilePath);
+	}
+
+	int32 PayloadsHydrated = 0;
 
-		int32 PayloadsHydrated = 0;
+	IVirtualizationSystem& System = IVirtualizationSystem::Get();
 
-		for (const FIoHash& Id : VirtualizedPayloads)
+	for (const FIoHash& Id : VirtualizedPayloads)
+	{
+		FCompressedBuffer Payload = System.PullData(Id);
+		if (!Payload.IsNull())
 		{
-			FCompressedBuffer Payload = System.PullData(Id);
-			if (!Payload.IsNull())
+			if (OutBuilder.UpdatePayloadAsLocal(Id, Payload))
 			{
-				if (Builder.UpdatePayloadAsLocal(Id, Payload))
-				{
-					PayloadsHydrated++;
-				}
-				else
-				{
-					FText Message = FText::Format(LOCTEXT("VAHydration_UpdateStatusFailed", "Unable to update the status for the payload '{0}' in the package '{1}'"),
-						FText::FromString(LexToString(Id)),
-						FText::FromString(FilePath));
-					OutErrors.Add(Message);
-
-					return;
-				}
+				PayloadsHydrated++;
 			}
 			else
 			{
-				FText Message = FText::Format(LOCTEXT("VAHydration_PullFailed", "Unable to pull the data for the payload '{0}' for the package '{1}'"),
+				FText Message = FText::Format(LOCTEXT("VAHydration_UpdateStatusFailed", "Unable to update the status for the payload '{0}' in the package '{1}'"),
 					FText::FromString(LexToString(Id)),
 					FText::FromString(FilePath));
 				OutErrors.Add(Message);
 
-				return;
+				return false;
 			}
 		}
+		else
+		{
+			FText Message = FText::Format(LOCTEXT("VAHydration_PullFailed", "Unable to pull the data for the payload '{0}' for the package '{1}'"),
+				FText::FromString(LexToString(Id)),
+				FText::FromString(FilePath));
+			OutErrors.Add(Message);
+
+			return false;
+		}
+	}
+
+	return PayloadsHydrated > 0;
+}
+
+void RehydratePackages(TConstArrayView<FString> PackagePaths, TArray<FText>& OutErrors)
+{
+	TRACE_CPUPROFILER_EVENT_SCOPE(UE::Virtualization::RehydratePackagesOnDisk);
+
+	IVirtualizationSystem& System = IVirtualizationSystem::Get();
 
-		if (PayloadsHydrated)
+	if (!System.IsEnabled())
+	{
+		return;
+	}
+
+	FScopedSlowTask Progress(1.0f, LOCTEXT("VAHydration_TaskOnDisk", "Re-hydrating Assets On Disk..."));
+	Progress.MakeDialog();
+
+	for (const FString& FilePath : PackagePaths)
+	{
+		FPackageTrailer Trailer;
+		FPackageTrailerBuilder Builder;
+
+		if (TryRehydrateBuilder(FilePath, Trailer, Builder, OutErrors))
 		{
 			const FString NewPackagePath = DuplicatePackageWithNewTrailer(FilePath, Trailer, Builder, OutErrors);
 
@@ -112,7 +206,7 @@ void RehydratePackages(const TArray<FString>& Packages, TArray<FText>& OutErrors
 				}
 			}
 
-			if (CanWriteToFile(FilePath))                     
+			if (CanWriteToFile(FilePath))
 			{
 				if (!IFileManager::Get().Move(*FilePath, *NewPackagePath))
 				{
@@ -136,6 +230,105 @@ void RehydratePackages(const TArray<FString>& Packages, TArray<FText>& OutErrors
 	}
 }
 
+void RehydratePackages(TConstArrayView<FString> PackagePaths, uint64 PaddingAlignment, TArray<FText>& OutErrors, TArray<FSharedBuffer>& OutPackages, TArray<FRehydrationInfo>* OutInfo)
+{
+	TRACE_CPUPROFILER_EVENT_SCOPE(UE::Virtualization::RehydratePackagesInMem);
+
+	IVirtualizationSystem& System = IVirtualizationSystem::Get();
+
+	FScopedSlowTask Progress(1.0f, LOCTEXT("VAHydration_TaskInMem", "Re-hydrating Assets In Memory..."));
+	Progress.MakeDialog();
+
+	OutPackages.Empty(PackagePaths.Num());
+	if (OutInfo != nullptr)
+	{
+		OutInfo->Empty(PackagePaths.Num());
+	}
+
+	for (const FString& FilePath : PackagePaths)
+	{
+		FUniqueBuffer FileBuffer;
+
+		FPackageTrailer Trailer;
+		FPackageTrailerBuilder Builder;
+
+		if (System.IsEnabled() && TryRehydrateBuilder(FilePath, Trailer, Builder, OutErrors))
+		{
+			TUniquePtr<FArchive> FileAr = OpenFileForReading(FilePath, OutErrors);
+			if (!FileAr)
+			{
+				return;
+			}
+
+			const int64 OriginalFileLength = FileAr->TotalSize();
+			const int64 TruncatedFileLength = OriginalFileLength - Trailer.GetTrailerLength();
+			const int64 RehydratedTrailerLength = (int64)Builder.CalculateTrailerLength();
+			const int64 RehydratedFileLength = TruncatedFileLength + RehydratedTrailerLength;
+			const int64 BufferLength = Align(RehydratedFileLength, PaddingAlignment);
+
+			FileBuffer = FUniqueBuffer::Alloc((uint64)BufferLength);
+
+			if (!ReadData(FileAr, FileBuffer.GetData(), TruncatedFileLength, FilePath, OutErrors))
+			{
+				return;
+			}
+
+			FFixedBufferWriterArchive Ar((uint8*)FileBuffer.GetData() + TruncatedFileLength, RehydratedTrailerLength);
+			if (!Builder.BuildAndAppendTrailer(nullptr, Ar))
+			{
+				FText Message = FText::Format(LOCTEXT("VAHydration_TrailerFailed", "Failed to create a new trailer for file {0}"),
+					FText::FromString(FilePath));
+
+				OutErrors.Add(MoveTemp(Message));
+
+				return;
+			}
+
+			check(Ar.GetRemainingSpace() == 0);	
+
+			if (OutInfo != nullptr)
+			{
+				FRehydrationInfo& Info = OutInfo->AddDefaulted_GetRef();
+				Info.OriginalSize = OriginalFileLength;
+				Info.RehydratedSize = RehydratedFileLength;
+				Info.NumPayloadsRehydrated = Builder.GetNumLocalPayloads() - Trailer.GetNumPayloads(EPayloadStorageType::Local);
+			}
+		}
+		else
+		{
+			TUniquePtr<FArchive> FileAr = OpenFileForReading(FilePath, OutErrors);
+			if (!FileAr)
+			{
+				return;
+			}
+
+			const int64 FileLength = FileAr->TotalSize();
+			const int64 BufferLength = Align(FileLength, PaddingAlignment);
+
+			FileBuffer = FUniqueBuffer::Alloc((uint64)BufferLength);
+
+			if (!ReadData(FileAr, FileBuffer.GetData(), FileLength, FilePath, OutErrors))
+			{
+				return;
+			}
+
+			if (OutInfo != nullptr)
+			{
+				FRehydrationInfo& Info = OutInfo->AddDefaulted_GetRef();
+				Info.OriginalSize = FileLength;
+				Info.RehydratedSize = FileLength;
+				Info.NumPayloadsRehydrated = 0;
+			}
+		}
+
+		OutPackages.Add(FileBuffer.MoveToShared());
+	}
+
+	// Make sure the arrays that we return have an entry per requested file
+	check(PackagePaths.Num() == OutPackages.Num());
+	check(OutInfo == nullptr || PackagePaths.Num() == OutInfo->Num());
+}
+
 } // namespace UE::Virtualization
 
 #undef LOCTEXT_NAMESPACE
diff --git a/Engine/Source/Developer/Virtualization/Private/PackageRehydrationProcess.h b/Engine/Source/Developer/Virtualization/Private/PackageRehydrationProcess.h
--- a/Engine/Source/Developer/Virtualization/Private/PackageRehydrationProcess.h
+++ b/Engine/Source/Developer/Virtualization/Private/PackageRehydrationProcess.h
@@ -6,9 +6,16 @@
 #include "Containers/UnrealString.h"
 #include "Internationalization/Text.h"
 
+class FSharedBuffer;
+struct FRehydrationInfo;
+
 namespace UE::Virtualization
 {
 
-void RehydratePackages(const TArray<FString>& Packages, TArray<FText>& OutErrors);
+/** This version will rehydrate a package and replace it on disk */
+void RehydratePackages(TConstArrayView<FString> PackagePaths, TArray<FText>& OutErrors);
+
+/** This version will rehydrate a package but to a memory buffer, the original package on disk will remain untouched */
+void RehydratePackages(TConstArrayView<FString> PackagePaths, uint64 PaddingAlignment, TArray<FText>& OutErrors, TArray<FSharedBuffer>& OutPackages, TArray<FRehydrationInfo>* OutInfo);
 
 } // namespace UE::Virtualization
diff --git a/Engine/Source/Developer/Virtualization/Private/PackageVirtualizationProcess.cpp b/Engine/Source/Developer/Virtualization/Private/PackageVirtualizationProcess.cpp
--- a/Engine/Source/Developer/Virtualization/Private/PackageVirtualizationProcess.cpp
+++ b/Engine/Source/Developer/Virtualization/Private/PackageVirtualizationProcess.cpp
@@ -154,7 +154,7 @@ private:
 	TMap<FIoHash, FPayloadData> PayloadLookupTable;
 };
 
-void VirtualizePackages(const TArray<FString>& FilesToSubmit, TArray<FText>& OutErrors)
+void VirtualizePackages(TConstArrayView<FString> PackagePaths, TArray<FText>& OutErrors)
 {
 	TRACE_CPUPROFILER_EVENT_SCOPE(UE::Virtualization::VirtualizePackages);
 
@@ -181,10 +181,10 @@ void VirtualizePackages(const TArray<FString>& FilesToSubmit, TArray<FText>& Out
 		bool bWasTrailerUpdated = false;
 	};
 
-	UE_LOG(LogVirtualization, Display, TEXT("Considering %d file(s) for virtualization"), FilesToSubmit.Num());
+	UE_LOG(LogVirtualization, Display, TEXT("Considering %d file(s) for virtualization"), PackagePaths.Num());
 
 	TArray<FPackageInfo> Packages;
-	Packages.Reserve(FilesToSubmit.Num());
+	Packages.Reserve(PackagePaths.Num());
 
 	Progress.EnterProgressFrame(1.0f);
 
@@ -193,7 +193,7 @@ void VirtualizePackages(const TArray<FString>& FilesToSubmit, TArray<FText>& Out
 	int64 TotalPackagesFound = 0;
 	int64 TotalPackageTrailersFound = 0;
 	int64 TotalPayloadsToVirtualize = 0;
-	for (const FString& AbsoluteFilePath : FilesToSubmit)
+	for (const FString& AbsoluteFilePath : PackagePaths)
 	{
 		FPackagePath PackagePath = FPackagePath::FromLocalPath(AbsoluteFilePath);
 
diff --git a/Engine/Source/Developer/Virtualization/Private/PackageVirtualizationProcess.h b/Engine/Source/Developer/Virtualization/Private/PackageVirtualizationProcess.h
--- a/Engine/Source/Developer/Virtualization/Private/PackageVirtualizationProcess.h
+++ b/Engine/Source/Developer/Virtualization/Private/PackageVirtualizationProcess.h
@@ -9,6 +9,6 @@
 namespace UE::Virtualization
 {
 
-void VirtualizePackages(const TArray<FString>& FilesToSubmit, TArray<FText>& OutErrors);
+void VirtualizePackages(TConstArrayView<FString> PackagePaths, TArray<FText>& OutErrors);
 
 } // namespace UE::Virtualization
diff --git a/Engine/Source/Developer/Virtualization/Private/VirtualizationManager.cpp b/Engine/Source/Developer/Virtualization/Private/VirtualizationManager.cpp
--- a/Engine/Source/Developer/Virtualization/Private/VirtualizationManager.cpp
+++ b/Engine/Source/Developer/Virtualization/Private/VirtualizationManager.cpp
@@ -654,14 +654,14 @@ EQueryResult FVirtualizationManager::QueryPayloadStatuses(TArrayView<const FIoHa
 	return EQueryResult::Success;
 }
 
-EVirtualizationResult FVirtualizationManager::TryVirtualizePackages(const TArray<FString>& FilesToVirtualize, TArray<FText>& OutDescriptionTags, TArray<FText>& OutErrors)
+EVirtualizationResult FVirtualizationManager::TryVirtualizePackages(TConstArrayView<FString> PackagePaths, TArray<FText>& OutDescriptionTags, TArray<FText>& OutErrors)
 {
 	OutDescriptionTags.Reset();
 	OutErrors.Reset();
 
 	if (IsEnabled() && IsPushingEnabled(EStorageType::Persistent))
 	{
-		UE::Virtualization::VirtualizePackages(FilesToVirtualize, OutErrors);
+		UE::Virtualization::VirtualizePackages(PackagePaths, OutErrors);
 
 		if (OutErrors.IsEmpty() && !VirtualizationProcessTag.IsEmpty())
 		{
@@ -673,11 +673,20 @@ EVirtualizationResult FVirtualizationManager::TryVirtualizePackages(const TArray
 	return OutErrors.IsEmpty() ? EVirtualizationResult::Success : EVirtualizationResult::Failed;
 }
 
-ERehydrationResult FVirtualizationManager::TryRehydratePackages(const TArray<FString>& Packages, TArray<FText>& OutErrors)
+ERehydrationResult FVirtualizationManager::TryRehydratePackages(TConstArrayView<FString> PackagePaths, TArray<FText>& OutErrors)
 {
 	OutErrors.Reset();
 
-	UE::Virtualization::RehydratePackages(Packages, OutErrors);
+	UE::Virtualization::RehydratePackages(PackagePaths, OutErrors);
+
+	return OutErrors.IsEmpty() ? ERehydrationResult::Success : ERehydrationResult::Failed;
+}
+
+ERehydrationResult FVirtualizationManager::TryRehydratePackages(TConstArrayView<FString> PackagePaths, uint64 PaddingAlignment, TArray<FText>& OutErrors, TArray<FSharedBuffer>& OutPackages, TArray<FRehydrationInfo>* OutInfo)
+{
+	OutErrors.Reset();
+
+	UE::Virtualization::RehydratePackages(PackagePaths, PaddingAlignment, OutErrors, OutPackages, OutInfo);
 
 	return OutErrors.IsEmpty() ? ERehydrationResult::Success : ERehydrationResult::Failed;
 }
diff --git a/Engine/Source/Developer/Virtualization/Private/VirtualizationManager.h b/Engine/Source/Developer/Virtualization/Private/VirtualizationManager.h
--- a/Engine/Source/Developer/Virtualization/Private/VirtualizationManager.h
+++ b/Engine/Source/Developer/Virtualization/Private/VirtualizationManager.h
@@ -125,8 +125,10 @@ private:
 
 	virtual EQueryResult QueryPayloadStatuses(TArrayView<const FIoHash> Ids, EStorageType StorageType, TArray<EPayloadStatus>& OutStatuses) override;
 
-	virtual EVirtualizationResult TryVirtualizePackages(const TArray<FString>& FilesToVirtualize, TArray<FText>& OutDescriptionTags, TArray<FText>& OutErrors) override;
-	virtual ERehydrationResult TryRehydratePackages(const TArray<FString>& Packages, TArray<FText>& OutErrors) override;
+	virtual EVirtualizationResult TryVirtualizePackages(TConstArrayView<FString> PackagePaths, TArray<FText>& OutDescriptionTags, TArray<FText>& OutErrors) override;
+
+	virtual ERehydrationResult TryRehydratePackages(TConstArrayView<FString> PackagePaths, TArray<FText>& OutErrors) override;
+	virtual ERehydrationResult TryRehydratePackages(TConstArrayView<FString> PackagePaths, uint64 PaddingAlignment, TArray<FText>& OutErrors, TArray<FSharedBuffer>& OutPackages, TArray<FRehydrationInfo>* OutInfo) override;
 
 	virtual void DumpStats() const override;
 
diff --git a/Engine/Source/Runtime/Core/Private/Virtualization/VirtualizationSystem.cpp b/Engine/Source/Runtime/Core/Private/Virtualization/VirtualizationSystem.cpp
--- a/Engine/Source/Runtime/Core/Private/Virtualization/VirtualizationSystem.cpp
+++ b/Engine/Source/Runtime/Core/Private/Virtualization/VirtualizationSystem.cpp
@@ -108,7 +108,7 @@ public:
 		return EQueryResult::Failure_NotImplemented;
 	}
 
-	virtual EVirtualizationResult TryVirtualizePackages(const TArray<FString>& FilesToVirtualize, TArray<FText>& OutDescriptionTags, TArray<FText>& OutErrors) override
+	virtual EVirtualizationResult TryVirtualizePackages(TConstArrayView<FString> PackagePaths, TArray<FText>& OutDescriptionTags, TArray<FText>& OutErrors) override
 	{
 		OutDescriptionTags.Reset();
 		OutErrors.Reset();
@@ -118,7 +118,16 @@ public:
 		return EVirtualizationResult::Failed;
 	}
 
-	virtual ERehydrationResult TryRehydratePackages(const TArray<FString>& Packages, TArray<FText>& OutErrors) override
+	virtual ERehydrationResult TryRehydratePackages(TConstArrayView<FString> PackagePaths, TArray<FText>& OutErrors) override
+	{
+		OutErrors.Reset();
+
+		OutErrors.Add(FText::FromString(TEXT("Calling ::TryRehydratePackages on FNullVirtualizationSystem")));
+
+		return ERehydrationResult::Failed;
+	}
+
+	virtual ERehydrationResult TryRehydratePackages(TConstArrayView<FString> PackagePaths, uint64 PaddingAlignment, TArray<FText>& OutErrors, TArray<FSharedBuffer>& OutPackages, TArray<FRehydrationInfo>* OutInfo) override
 	{
 		OutErrors.Reset();
 
diff --git a/Engine/Source/Runtime/Core/Public/Virtualization/VirtualizationSystem.h b/Engine/Source/Runtime/Core/Public/Virtualization/VirtualizationSystem.h
--- a/Engine/Source/Runtime/Core/Public/Virtualization/VirtualizationSystem.h
+++ b/Engine/Source/Runtime/Core/Public/Virtualization/VirtualizationSystem.h
@@ -102,6 +102,17 @@ enum class ERehydrationResult : uint8
 	Failed,
 };
 
+/** Info about a rehydration operation */
+struct FRehydrationInfo
+{
+	/** Size of the package before it was rehydrated */
+	uint64 OriginalSize = 0;
+	/** Size of the package after rehydration, this does not including padding if any was applied */
+	uint64 RehydratedSize = 0;
+	/** The number of payloads that needed to be rehydrated */
+	int32 NumPayloadsRehydrated = 0;
+};
+
 /** 
  * This interface can be implemented and passed to a FPushRequest as a way of providing the payload 
  * to the virtualization system for a push operation but deferring the loading of the payload from 
@@ -582,7 +593,7 @@ public:
 	 * Runs the virtualization process on a set of packages. All of the packages will be parsed and any found to be containing locally stored
 	 * payloads will have them removed but before they are removed they will be pushed to persistent storage.
 	 * 
-	 * @param FilesToVirtualize		An array of file paths to packages that should be virtualized. If a path resolves to a file that is not 
+	 * @param PackagePaths			An array of file paths to packages that should be virtualized. If a path resolves to a file that is not 
 	 *								a valid package then it will be silently skipped and will not be considered an error.
 	 * @param OutDescriptionTags	The process may produce description tags associated with the packages which will be placed in this array.
 	 *								These tags can be used to improve logging, or be appended to change list descriptions etc. Note that the 
@@ -595,22 +606,40 @@ public:
 	 * @return						A EVirtualizationResult enum with the status of the process. If the status is not EVirtualizationResult::Success
 	 *								then the parameter OutErrors should contain at least one entry.
 	 */
-	virtual EVirtualizationResult TryVirtualizePackages(const TArray<FString>& FilesToVirtualize, TArray<FText>& OutDescriptionTags, TArray<FText>& OutErrors) = 0;
+	virtual EVirtualizationResult TryVirtualizePackages(TConstArrayView<FString> PackagePaths, TArray<FText>& OutDescriptionTags, TArray<FText>& OutErrors) = 0;
 	
 	/**
-	 * Runs the re-hydration process on a set of packages. This involves downloading virtualized payloads and placing them back in the trailer of
+	 * Runs the rehydration process on a set of packages. This involves downloading virtualized payloads and placing them back in the trailer of
 	 * the given packages.
 	 * 
-	 * @param Packages	An array containing the absolute file paths of packages. It is assumed that the packages have already been checked out
-	 *					of source control (if applicable) and will be writable.
-	 * @param OutErrors	Any error encountered during the process will be added here. If any error is added to the array then it
-	 *					can be assumed that the process will return false. Note that the array will be emptied before the process
-	 *					is run and will not contain any pre-existing entries.
+	 * @param PackagePaths	An array containing the absolute file paths of packages. It is assumed that the packages have already been checked out
+	 *						of source control (if applicable) and will be writable.
+	 * @param OutErrors		Any error encountered during the process will be added here. If any error is added to the array then it
+	 *						can be assumed that the process will return false. Note that the array will be emptied before the process
+	 *						is run and will not contain any pre-existing entries.
 	 * 
 	 * @return	A ERehydrationResult enum with the status of the process. If the status indicates any sort of failure then OutErrors should
 	 *			contain at least one entry.
 	 */
-	virtual ERehydrationResult TryRehydratePackages(const TArray<FString>& Packages, TArray<FText>& OutErrors) = 0;
+	virtual ERehydrationResult TryRehydratePackages(TConstArrayView<FString> PackagePaths, TArray<FText>& OutErrors) = 0;
+
+	/**
+	 * Rehydrates a number of packages into memory buffers.
+	 * Note that if a package does not require rehydration we will still return the package in a memory buffer but it will be the same
+	 * as the package on disk.
+	 * 
+	 * @param PackagePaths		An array containing the absolute file paths of packages.
+	 * @param PaddingAlignment	Byte alignment to pad each package buffer too, a value of 0 will result in the buffers being the same size as the packages
+	 * @param OutErrors			Any errors encountered while rehydration will be added here
+	 * @param OutPackages		The rehydrated packages as memory buffers. Each entry should match the corresponding entry in PackagePaths. This array is
+	 *							only guaranteed to be correct if the method returns ERehydrationResult::Success.
+	 * @param OutInfo			Information about thge rehydration process, each entry should match the corresponding entry in PackagePaths assuming that
+	 *							the method returns ERehydrationResult::Success. This parameter is optional, if the information is not required then
+	 *							pass in nullptr to skip.
+	 * @return					ERehydrationResult::Success if the rehydration suceeeds and OutPackages/OutInfocan be trued, otherwise it will return an
+	 *							error value. @see ERehydrationResult
+	 */
+	virtual ERehydrationResult TryRehydratePackages(TConstArrayView<FString> PackagePaths, uint64 PaddingAlignment, TArray<FText>& OutErrors, TArray<FSharedBuffer>& OutPackages, TArray<FRehydrationInfo>* OutInfo) = 0;
 
 	/** When called the system should write any performance stats that it has been gathering to the log file */
 	virtual void DumpStats() const = 0;
