Updated all VA related commandlines to work without '-' although the previous version of the command line options will continue to work.

#rnx

- It was pointed out that a command line option containing '-' characters could confuse parsers and it is not worth the risk to keep them.
- Now we accept versions of each command line without the '-' character which will be the version that we document for users in 5.1

diff --git a/Engine/Source/Developer/Virtualization/Private/VirtualizationManager.cpp b/Engine/Source/Developer/Virtualization/Private/VirtualizationManager.cpp
--- a/Engine/Source/Developer/Virtualization/Private/VirtualizationManager.cpp
+++ b/Engine/Source/Developer/Virtualization/Private/VirtualizationManager.cpp
@@ -99,6 +99,27 @@ bool LexTryParseString(EPackageFilterMode& OutValue, FStringView Buffer)
 	return false;
 }
 
+/** 
+ * Utility to check if either cmd is present in the command line. Useful when transitioning from one
+ * command line to another. 
+ */
+static bool IsCmdLineSet(const TCHAR* Cmd, const TCHAR* AlternativeCmd = nullptr)
+{
+	const TCHAR* CmdLine = FCommandLine::Get();
+
+	if (FParse::Param(CmdLine, Cmd))
+	{
+		return true;
+	}
+
+	if (AlternativeCmd != nullptr && FParse::Param(CmdLine, AlternativeCmd))
+	{
+		return true;
+	}
+
+	return false;
+}
+
 /* Utility function for building up a lookup table of all available IBackendFactory interfaces*/
 FVirtualizationManager::FRegistedFactories FindBackendFactories()
 {
@@ -981,13 +1002,13 @@ void FVirtualizationManager::ApplySettingsFromConfigFiles(const FConfigFile& Con
 
 void FVirtualizationManager::ApplySettingsFromFromCmdline()
 {
-	if (!bLazyInitConnections && FParse::Param(FCommandLine::Get(), TEXT("VA-LazyInitConnections")))
+	if (!bLazyInitConnections && IsCmdLineSet(TEXT("VALazyInitConnections"), TEXT("VA-LazyInitConnections")))
 	{
 		bLazyInitConnections = true;
 		UE_LOG(LogVirtualization, Display, TEXT("Cmdline has set the virtualization system backends to lazy init their connections"));
 	}
 
-	if (bAllowPackageVirtualization && FParse::Param(FCommandLine::Get(), TEXT("VA-SkipPkgVirtualization")))
+	if (bAllowPackageVirtualization && IsCmdLineSet(TEXT("VASkipPkgVirtualization"), TEXT("VA-SkipPkgVirtualization")))
 	{
 		bAllowPackageVirtualization = false;
 		UE_LOG(LogVirtualization, Warning, TEXT("The virtualization process has been disabled via the command line"));
@@ -1005,27 +1026,27 @@ void FVirtualizationManager::ApplySettingsFromCVar()
 
 void FVirtualizationManager::ApplyDebugSettingsFromFromCmdline()
 {
-	if (FParse::Param(FCommandLine::Get(), TEXT("VA-SingleThreaded")))
+	if (IsCmdLineSet(TEXT("VASingleThreaded"), TEXT("VA-SingleThreaded")))
 	{
 		DebugValues.bSingleThreaded = true;
 		UE_LOG(LogVirtualization, Warning, TEXT("Cmdline has set the virtualization system to run single threaded"));
 	}
 
-	if (FParse::Param(FCommandLine::Get(), TEXT("VA-ValidatePushes")))
+	if (IsCmdLineSet(TEXT("VAValidatePushes"), TEXT("VA-ValidatePushes")))
 	{
 		DebugValues.bValidateAfterPush = true;
 		UE_LOG(LogVirtualization, Warning, TEXT("Cmdline has set the virtualization system to pull each payload after pushing to either local or persistent storage"));
 	}
 
 	FString CmdlineGraphName;
-	if (FParse::Value(FCommandLine::Get(), TEXT("-VA-BackendGraph="), CmdlineGraphName))
+	if (IsCmdLineSet(TEXT("VABackendGraph"), TEXT("VA-BackendGraph")))
 	{
 		UE_LOG(LogVirtualization, Display, TEXT("Backend graph overriden from the cmdline: '%s'"), *CmdlineGraphName);
 		BackendGraphName = CmdlineGraphName;
 	}
 
 	FString MissOptions;
-	if (FParse::Value(FCommandLine::Get(), TEXT("-VA-MissBackends="), MissOptions))
+	if (IsCmdLineSet(TEXT("VAMissBackends"), TEXT("VA-MissBackends")))
 	{
 		MissOptions.ParseIntoArray(DebugValues.MissBackends, TEXT("+"), true);
 
@@ -1037,7 +1058,7 @@ void FVirtualizationManager::ApplyDebugSettingsFromFromCmdline()
 	}
 
 	DebugValues.MissChance = 0.0f;
-	if (FParse::Value(FCommandLine::Get(), TEXT("-VA-MissChance="), DebugValues.MissChance))
+	if (IsCmdLineSet(TEXT("VAMissChance"), TEXT("VA-MissChance")))
 	{
 		DebugValues.MissChance = FMath::Clamp(DebugValues.MissChance, 0.0f, 100.0f);
 
diff --git a/Engine/Source/Runtime/Core/Private/Virtualization/VirtualizationSystem.cpp b/Engine/Source/Runtime/Core/Private/Virtualization/VirtualizationSystem.cpp
--- a/Engine/Source/Runtime/Core/Private/Virtualization/VirtualizationSystem.cpp
+++ b/Engine/Source/Runtime/Core/Private/Virtualization/VirtualizationSystem.cpp
@@ -161,6 +161,27 @@ public:
 
 TUniquePtr<IVirtualizationSystem> GVirtualizationSystem = nullptr;
 
+/**
+ * Utility to check if either cmd is present in the command line. Useful when transitioning from one
+ * command line to another.
+ */
+static bool IsCmdLineSet(const TCHAR* Cmd, const TCHAR* AlternativeCmd = nullptr)
+{
+	const TCHAR* CmdLine = FCommandLine::Get();
+
+	if (FParse::Param(CmdLine, Cmd))
+	{
+		return true;
+	}
+
+	if (AlternativeCmd != nullptr && FParse::Param(CmdLine, AlternativeCmd))
+	{
+		return true;
+	}
+
+	return false;
+}
+
 /** Utility function for finding a IVirtualizationSystemFactory for a given system name */
 Private::IVirtualizationSystemFactory* FindFactory(FName SystemName)
 {
@@ -183,7 +204,7 @@ bool ShouldLazyInitializeSystem(const FConfigFile& ConfigFile)
 	UE_LOG(LogVirtualization, Display, TEXT("The virtualization system will lazy initialize due to code"));
 	return true;
 #else
-	if (FParse::Param(FCommandLine::Get(), TEXT("VA-LazyInit")))
+	if (IsCmdLineSet(TEXT("VALazyInit"), TEXT("VA-LazyInit")))
 	{
 		UE_LOG(LogVirtualization, Display, TEXT("The virtualization system will lazy initialize due to the command line"));
 		return true;
@@ -218,7 +239,7 @@ FName FindSystemToMount(const FConfigFile& ConfigFile)
 	UE_LOG(LogVirtualization, Display, TEXT("The virtualization system has been disabled by code"));
 	return FName();
 #else
-	if (FParse::Param(FCommandLine::Get(), TEXT("VA-Disable")))
+	if (IsCmdLineSet(TEXT("VADisable"), TEXT("VA-Disable")))
 	{
 		UE_LOG(LogVirtualization, Display, TEXT("The virtualization system has been disabled by the command line"));
 		return FName();
