Extending the VA pulling API to allow for batch pulls

#rnx

- The caller can now request more than one payload to be pulled froim the system at a time. In theory this will allow backends to pull data quicker.
-- The file system backend works with the new system but makes no attempt to take advantage as the backend is not intended for serious production use.
-- The DDC backend should be taking full advantage of the new batching API
-- Note that although the source control API does attempt to take advantage of batching, the internal source control API implementation is not doing so. This will be addressed in a future submit.
- This does make the pulling logic a bit more complicated in FVirtualizationManager as we need to deal with some payloads being found in the first backend, some in the second and some in the third etc.
-- To help deal with this a new class FPullRequestCollection has been added which abstracts a lot of complexity away.

diff --git a/Engine/Source/Developer/Virtualization/Private/VirtualizationDDCBackend.cpp b/Engine/Source/Developer/Virtualization/Private/VirtualizationDDCBackend.cpp
--- a/Engine/Source/Developer/Virtualization/Private/VirtualizationDDCBackend.cpp
+++ b/Engine/Source/Developer/Virtualization/Private/VirtualizationDDCBackend.cpp
@@ -131,36 +131,38 @@ bool FDDCBackend::PushData(TArrayView<FPushRequest> Requests)
 	return bWasSuccess;
 }
 
-FCompressedBuffer FDDCBackend::PullData(const FIoHash& Id)
+bool FDDCBackend::PullData(TArrayView<FPullRequest> Requests)
 {
 	TRACE_CPUPROFILER_EVENT_SCOPE(FDDCBackend::PullData);
 
 	UE::DerivedData::ICache& Cache = UE::DerivedData::GetCache();
 
-	UE::DerivedData::FCacheKey Key;
-	Key.Bucket = Bucket; 
-	Key.Hash = Id;
+	UE::DerivedData::FRequestOwner Owner(UE::DerivedData::EPriority::Normal);
 
-	UE::DerivedData::FRequestOwner Owner(UE::DerivedData::EPriority::Blocking);
+	for (FPullRequest& Request : Requests)
+	{
+		UE::DerivedData::FRequestBarrier Barrier(Owner);
 
-	FCompressedBuffer ResultData;
-	UE::DerivedData::EStatus ResultStatus;
+		UE::DerivedData::FCacheKey Key;
+		Key.Bucket = Bucket;
+		Key.Hash = Request.GetIdentifier();
 
-	auto Callback = [&Id, &ResultData, &ResultStatus](UE::DerivedData::FCacheGetResponse&& Response)
-	{
-		ResultStatus = Response.Status;
-		if (ResultStatus == UE::DerivedData::EStatus::Ok)
+		auto Callback = [&Request](UE::DerivedData::FCacheGetResponse&& Response)
 		{
-			ResultData = Response.Record.GetValue(ToDerivedDataValueId(Id)).GetData();
-		}
-	};
+			if (Response.Status == UE::DerivedData::EStatus::Ok)
+			{
+				UE::DerivedData::FValueId ValueId = ToDerivedDataValueId(Request.GetIdentifier());
+				Request.SetPayload(Response.Record.GetValue(ValueId).GetData());
+			}
+		};
 
-	// TODO: Improve the name when we start passing more context to this function
-	Cache.Get({{{TEXT("Mirage")}, Key, TransferPolicy}}, Owner, MoveTemp(Callback));
+		// TODO: Improve the name when we start passing more context to this function
+		Cache.Get({ {{TEXT("Mirage")}, Key, TransferPolicy} }, Owner, MoveTemp(Callback));
+	}
 
 	Owner.Wait();
 
-	return ResultData;
+	return true;
 }
 
 bool FDDCBackend::DoesPayloadExist(const FIoHash& Id)
diff --git a/Engine/Source/Developer/Virtualization/Private/VirtualizationDDCBackend.h b/Engine/Source/Developer/Virtualization/Private/VirtualizationDDCBackend.h
--- a/Engine/Source/Developer/Virtualization/Private/VirtualizationDDCBackend.h
+++ b/Engine/Source/Developer/Virtualization/Private/VirtualizationDDCBackend.h
@@ -39,8 +39,7 @@ private:
 	virtual EConnectionStatus OnConnect() override;
 
 	virtual bool PushData(TArrayView<FPushRequest> Requests) override;
-
-	virtual FCompressedBuffer PullData(const FIoHash& Id) override;
+	virtual bool PullData(TArrayView<FPullRequest> Requests) override;
 	
 	virtual bool DoesPayloadExist(const FIoHash& Id) override;
 
diff --git a/Engine/Source/Developer/Virtualization/Private/VirtualizationFileBackend.cpp b/Engine/Source/Developer/Virtualization/Private/VirtualizationFileBackend.cpp
--- a/Engine/Source/Developer/Virtualization/Private/VirtualizationFileBackend.cpp
+++ b/Engine/Source/Developer/Virtualization/Private/VirtualizationFileBackend.cpp
@@ -168,37 +168,59 @@ bool FFileSystemBackend::PushData(TArrayView<FPushRequest> Requests)
 	return ErrorCount == 0;
 }
 
-FCompressedBuffer FFileSystemBackend::PullData(const FIoHash& Id)
+bool FFileSystemBackend::PullData(TArrayView<FPullRequest> Requests)
 {
 	TRACE_CPUPROFILER_EVENT_SCOPE(FFileSystemBackend::PullData);
 
-	TStringBuilder<512> FilePath;
-	CreateFilePath(Id, FilePath);
-
-	// TODO: Should we allow the error severity to be configured via ini or just not report this case at all?
-	if (!IFileManager::Get().FileExists(FilePath.ToString()))
+	for (FPullRequest& Request : Requests)
 	{
-		UE_LOG(LogVirtualization, Verbose, TEXT("[%s] Does not contain the payload '%s'"), *GetDebugName(), *LexToString(Id));
-		return FCompressedBuffer();
-	}
+		TStringBuilder<512> FilePath;
+		CreateFilePath(Request.GetIdentifier(), FilePath);
 
-	TUniquePtr<FArchive> FileAr = OpenFileForReading(FilePath.ToString());
+		// TODO: Should we allow the error severity to be configured via ini or just not report this case at all?
+		if (!IFileManager::Get().FileExists(FilePath.ToString()))
+		{
+			UE_LOG(LogVirtualization, Verbose, TEXT("[%s] Does not contain the payload '%s'"), 
+				*GetDebugName(), 
+				*LexToString(Request.GetIdentifier()));
 
-	if (FileAr == nullptr)
-	{
-		TStringBuilder<MAX_SPRINTF> SystemErrorMsg;
-		Utils::GetFormattedSystemError(SystemErrorMsg);
+			continue;
+		}
 
-		UE_LOG(LogVirtualization, Error, TEXT("[%s] Failed to load payload '%s' from file '%s' due to system error: %s"),
-			*GetDebugName(),
-			*LexToString(Id),
-			FilePath.ToString(),
-			SystemErrorMsg.ToString());
+		TUniquePtr<FArchive> FileAr = OpenFileForReading(FilePath.ToString());
 
-		return FCompressedBuffer();
+		if (FileAr == nullptr)
+		{
+			TStringBuilder<MAX_SPRINTF> SystemErrorMsg;
+			Utils::GetFormattedSystemError(SystemErrorMsg);
+
+			UE_LOG(LogVirtualization, Error, TEXT("[%s] Failed to load payload '%s' from file '%s' due to system error: %s"),
+				*GetDebugName(),
+				*LexToString(Request.GetIdentifier()),
+				FilePath.ToString(),
+				SystemErrorMsg.ToString());
+
+			Request.SetError();
+
+			continue;
+		}
+
+		Request.SetPayload(FCompressedBuffer::Load(*FileAr));
+
+		if (FileAr->IsError())
+		{
+			UE_LOG(LogVirtualization, Error, TEXT("[%s] Failed to serialize payload '%s' from file '%s'"),
+				*GetDebugName(),
+				*LexToString(Request.GetIdentifier()),
+				FilePath.ToString());
+
+			Request.SetError();
+
+			continue;
+		}
 	}
 
-	return FCompressedBuffer::Load(*FileAr);
+	return true;
 }
 
 bool FFileSystemBackend::DoesPayloadExist(const FIoHash& Id)
diff --git a/Engine/Source/Developer/Virtualization/Private/VirtualizationFileBackend.h b/Engine/Source/Developer/Virtualization/Private/VirtualizationFileBackend.h
--- a/Engine/Source/Developer/Virtualization/Private/VirtualizationFileBackend.h
+++ b/Engine/Source/Developer/Virtualization/Private/VirtualizationFileBackend.h
@@ -42,8 +42,7 @@ private:
 	virtual EConnectionStatus OnConnect() override;
 
 	virtual bool PushData(TArrayView<FPushRequest> Requests) override;
-
-	virtual FCompressedBuffer PullData(const FIoHash& Id) override;
+	virtual bool PullData(TArrayView<FPullRequest> Requests) override;
 
 	virtual bool DoesPayloadExist(const FIoHash& Id) override;
 	
diff --git a/Engine/Source/Developer/Virtualization/Private/VirtualizationManager.cpp b/Engine/Source/Developer/Virtualization/Private/VirtualizationManager.cpp
--- a/Engine/Source/Developer/Virtualization/Private/VirtualizationManager.cpp
+++ b/Engine/Source/Developer/Virtualization/Private/VirtualizationManager.cpp
@@ -31,6 +31,8 @@ UE_REGISTER_VIRTUALIZATION_SYSTEM(UE::Virtualization::FVirtualizationManager, De
 	#define UE_VIRTUALIZATION_CONNECTION_LAZY_INIT 0
 #endif //UE_VIRTUALIZATION_CONNECTION_LAZY_INIT
 
+#define UE_INLINE_ALLOCATION_COUNT 4
+
 // TODO: Move to RegisterConsoleCommands
 static TAutoConsoleVariable<bool> CVarLazyInitConnections(
 	TEXT("VA.LazyInitConnections"),
@@ -83,6 +85,126 @@ private:
 	FCriticalSection* SyncObject;
 };
 
+/** 
+ * Utility class to help manage pull requests. When created it will remove invalid and duplicate requests so
+ * that backends do not need to worry about them.
+ * It will also prune the list of requests after each backend pull so that we only make requests for payloads 
+ * that have not yet been found and on request provide a list of payloads that should be cached.
+ * On destruction it will then write the results back to the original requests. Duplicate requests will share 
+ * references to the same payload in memory. Once destroyed all of the original requests should have a success
+ * or error status and none should be listed as pending.
+ * 
+ * NOTE: This is intended to be used by FVirtualizationManager::PullDataFromAllBackends only, hence doing a 
+ * few dangerous things like returning a TArrayView as we know it will not be misused. If this is ever taken
+ * into wider use we will have to revisit things like that.
+ */
+class FPullRequestCollection
+{
+public:
+	UE_NONCOPYABLE(FPullRequestCollection);
+
+	FPullRequestCollection() = delete;
+	FPullRequestCollection(TArrayView<FPullRequest> InRequests)
+		: OriginalRequests(InRequests)
+	{
+		CurrentRequests.Reserve(OriginalRequests.Num());
+
+		// Record each payload hash as we add it to CurrentRequests so that we only request duplicates once.
+		TSet<FIoHash, DefaultKeyFuncs<FIoHash>, TInlineSetAllocator<UE_INLINE_ALLOCATION_COUNT>> UniquePayloads;
+
+		for (FPullRequest& Request : OriginalRequests)
+		{
+			if (Request.GetIdentifier().IsZero())
+			{
+				Request.SetError();
+				UE_LOG(LogVirtualization, Error, TEXT("Attempting to pull a virtualized payload with an invalid FIoHash"));
+			}
+			else if (!UniquePayloads.Contains(Request.GetIdentifier()))
+			{
+				CurrentRequests.Add(Request);
+				UniquePayloads.Add(Request.GetIdentifier());
+			}
+		}
+	}
+
+	~FPullRequestCollection()
+	{
+		for (FPullRequest& Request : OriginalRequests)
+		{
+			if (FCompressedBuffer* Payload = LoadedPayloads.Find(Request.GetIdentifier()))
+			{
+				Request.SetPayload(*Payload);
+			}
+			else
+			{
+				Request.SetError();
+			}
+		}
+	}
+
+	/** 
+	 * Called after the requests from ::GetRequests have been pulled from a backend. Payloads that were
+	 * successfully pulled will be removed from the request list and added to the LoadedPayloads map 
+	 * so that they can be assigned to the original requests later.
+	 * 
+	 * @param Backend					The backend that the payloads were pulled from
+	 * @param bRequirePayloadsToCache	Do we need a list of payloads that should be cached
+	 * 
+	 * @return A list of requests that now need to be cached (if required)
+	 */
+	TArray<FPushRequest> OnPullCompleted(const IVirtualizationBackend& Backend, bool bRequirePayloadsToCache)
+	{
+		TArray<FPushRequest> PayloadsToCache;
+		if (bRequirePayloadsToCache)
+		{
+			PayloadsToCache.Reserve(CurrentRequests.Num());
+		}
+
+		for (int32 Index = 0; Index < CurrentRequests.Num();)
+		{
+			const FPullRequest& Request = CurrentRequests[Index];
+			if (Request.IsSuccess())
+			{
+				if (bRequirePayloadsToCache)
+				{
+					PayloadsToCache.Emplace(FPushRequest(Request.GetIdentifier(), Request.GetPayload(), FString()));
+				}
+
+				LoadedPayloads.Add(Request.GetIdentifier(), Request.GetPayload());
+
+				CurrentRequests.RemoveAtSwap(Index);
+
+				UE_LOG(LogVirtualization, VeryVerbose, TEXT("[%s] pulled payload '%s'"), *Backend.GetDebugName(), *LexToString(Request.GetIdentifier()));
+			}
+			else
+			{
+				++Index;
+			}
+		}
+
+		return PayloadsToCache;
+	}
+
+	/** Return the current list of requests that need to be made */
+	TArrayView<FPullRequest> GetRequests()
+	{
+		return CurrentRequests;
+	}
+
+	/** Returns if we still have payloads that have not yet been found */
+	bool HasRemainingRequests() const
+	{
+		return !CurrentRequests.IsEmpty();
+	}
+
+private:
+	TArrayView<FPullRequest> OriginalRequests;
+
+	TMap<FIoHash, FCompressedBuffer, TInlineSetAllocator<UE_INLINE_ALLOCATION_COUNT>> LoadedPayloads;
+	TArray<FPullRequest, TInlineAllocator<UE_INLINE_ALLOCATION_COUNT>> CurrentRequests;
+};
+
+
 bool LexTryParseString(EPackageFilterMode& OutValue, FStringView Buffer)
 {
 	if (Buffer == TEXT("OptOut"))
@@ -523,13 +645,15 @@ bool FVirtualizationManager::PushData(TArrayView<FPushRequest> Requests, EStorag
 		// Debug operation to validate that the payload we just pushed can be retrieved from storage
 		if (DebugValues.bValidateAfterPush && bResult == true && Backend->IsOperationSupported(IVirtualizationBackend::EOperations::Pull))
 		{
-			for (FPushRequest& Request : ValidatedRequests)
+			for (FPushRequest& PushRequest : ValidatedRequests)
 			{
-				FCompressedBuffer ValidationPayload = PullDataFromBackend(*Backend, Request.GetIdentifier());
-				checkf(	Request.GetIdentifier() == ValidationPayload.GetRawHash(),
+				FPullRequest PullRequest(PushRequest.GetIdentifier());
+				PullDataFromBackend(*Backend, MakeArrayView(&PullRequest, 1));
+
+				checkf(PushRequest.GetIdentifier() == PullRequest.GetPayload().GetRawHash(),
 						TEXT("[%s] Failed to pull payload '%s' after it was pushed to backend"),
 						*Backend->GetDebugName(),
-						*LexToString(Request.GetIdentifier()));
+						*LexToString(PushRequest.GetIdentifier()));
 			}
 		}
 	}
@@ -551,51 +675,51 @@ bool FVirtualizationManager::PushData(TArrayView<FPushRequest> Requests, EStorag
 	return IsCacheType(StorageType) ? ErrorCount < Backends.Num() : ErrorCount == 0;
 }
 
-FCompressedBuffer FVirtualizationManager::PullData(const FIoHash& Id)
+bool FVirtualizationManager::PullData(TArrayView<FPullRequest> Requests)
 {
 	TRACE_CPUPROFILER_EVENT_SCOPE(FVirtualizationManager::PullData);
 
-	if (Id.IsZero())
-	{
-		// TODO: See below, should errors here be fatal?
-		UE_LOG(LogVirtualization, Error, TEXT("Attempting to pull a virtualized payload with an invalid FIoHash"));
-		return FCompressedBuffer();
-	}
-
 	if (PullEnabledBackends.IsEmpty())
 	{
 		// TODO: See below, should errors here be fatal?
-		UE_LOG(LogVirtualization, Error, TEXT("Payload '%s' failed to be pulled as there are no backends mounted!'"), *LexToString(Id));
-		return FCompressedBuffer();
+		UE_LOG(LogVirtualization, Error, TEXT("Failed to pull payload(s) as there are no backends mounted!'"));
+		return false;
 	}
 
 	EnsureBackendConnections();
 
 	FConditionalScopeLock _(&DebugValues.ForceSingleThreadedCS, DebugValues.bSingleThreaded);
 
-	GetNotificationEvent().Broadcast(IVirtualizationSystem::PullBegunNotification, Id);
+	BroadcastEvent(Requests, IVirtualizationSystem::PullBegunNotification);
 
-	FCompressedBuffer Payload = PullDataFromAllBackends(Id);
-	
-	GetNotificationEvent().Broadcast(IVirtualizationSystem::PullEndedNotification, Id);
+	PullDataFromAllBackends(Requests);
 
-	if (!Payload.IsNull())
-	{
-		return Payload;
-	}
-	else
+	BroadcastEvent(Requests, IVirtualizationSystem::PullEndedNotification);
+
+	bool bSuccess = true;
+	for (const FPullRequest& Request : Requests)
 	{
-		// Broadcast the pull failed event to any listeners
-		GetNotificationEvent().Broadcast(IVirtualizationSystem::PullFailedNotification, Id);
+		// Report failed pulls of valid identifiers
+		if (!Request.IsSuccess() && !Request.GetIdentifier().IsZero())
+		{
+			GetNotificationEvent().Broadcast(IVirtualizationSystem::PullFailedNotification, Request.GetIdentifier());
 
-		// TODO: Maybe this should be a fatal error? If we keep it as an error we need to make sure any calling
-		// code handles it properly.
-		// Could be worth extending ::PullData to return error codes instead so we can make a better distinction 
-		// between the payload not being found in any of the backends and one or more of the backends failing.
-		UE_LOG(LogVirtualization, Error, TEXT("Payload '%s' failed to be pulled from any backend'"), *LexToString(Id));
+			// TODO: Maybe this should be a fatal error? If we keep it as an error we need to make sure any calling
+			// code handles it properly.
+			// Could be worth extending ::PullData to return error codes instead so we can make a better distinction 
+			// between the payload not being found in any of the backends and one or more of the backends failing.
+			UE_LOG(LogVirtualization, Error, TEXT("Payload '%s' failed to be pulled from any backend'"), *LexToString(Request.GetIdentifier()));
 
-		return FCompressedBuffer();
-	}
+			bSuccess = false;
+		}
+
+		if (Request.IsSuccess())
+		{
+			checkf(Request.GetIdentifier() == Request.GetPayload().GetRawHash(), TEXT("[%s] Invalid payload for '%s'"), *LexToString(Request.GetIdentifier()));
+		}
+	}	
+
+	return bSuccess;
 }
 
 EQueryResult FVirtualizationManager::QueryPayloadStatuses(TArrayView<const FIoHash> Ids, EStorageType StorageType, TArray<EPayloadStatus>& OutStatuses)
@@ -1540,7 +1664,7 @@ void FVirtualizationManager::EnsureBackendConnections()
 	}
 }
 
-void FVirtualizationManager::CachePayload(const FIoHash& Id, const FCompressedBuffer& Payload, const IVirtualizationBackend* BackendSource)
+void FVirtualizationManager::CachePayloads(TArrayView<FPushRequest> Requests, const IVirtualizationBackend* BackendSource)
 {
 	TRACE_CPUPROFILER_EVENT_SCOPE(FVirtualizationManager::CachePayload);
 
@@ -1549,41 +1673,61 @@ void FVirtualizationManager::CachePayload(const FIoHash& Id, const FCompressedBu
 	// We start caching at the first (assumed to be fastest) local cache backend. 
 	for (IVirtualizationBackend* BackendToCache : CacheStorageBackends)
 	{
+		// We stop once we reach the backend that the payloads were first pulled from
 		if (BackendToCache == BackendSource)
 		{
-			return; // No point going past BackendSource
+			return;
 		}
 
-		const bool bResult = TryCacheDataToBackend(*BackendToCache, Id, Payload);
-		UE_CLOG(	!bResult, LogVirtualization, Warning,
-					TEXT("Failed to cache payload '%s' to backend '%s'"),
-					*LexToString(Id),
-					*BackendToCache->GetDebugName());
+		const bool bResult = TryCacheDataToBackend(*BackendToCache, Requests);
+
+		if (!bResult)
+		{
+			for (const FPushRequest& Request : Requests)
+			{
+				UE_LOG(LogVirtualization, Warning, TEXT("Failed to cache payload '%s' to backend '%s'"), *LexToString(Request.GetIdentifier()), *BackendToCache->GetDebugName());
+			}
+		}
 
 		// Debug operation to validate that the payload we just cached can be retrieved from storage
 		if (DebugValues.bValidateAfterPush && bResult && BackendToCache->IsOperationSupported(IVirtualizationBackend::EOperations::Pull))
 		{
-			FCompressedBuffer PulledPayload = PullDataFromBackend(*BackendToCache, Id);
-			checkf(	Payload.GetRawHash() == PulledPayload.GetRawHash(), 
+			for (const FPushRequest& Request : Requests)
+			{
+				FPullRequest ValidationRequest(Request.GetIdentifier());
+				PullDataFromBackend(*BackendToCache, MakeArrayView(&ValidationRequest, 1));
+
+				checkf(Request.GetPayload().GetRawHash() == ValidationRequest.GetPayload().GetRawHash(),
 					TEXT("[%s] Failed to pull payload '%s' after it was cached to backend"),
-					*BackendToCache->GetDebugName(), 
-					*LexToString(Id));
+					*BackendToCache->GetDebugName(),
+					*LexToString(Request.GetIdentifier()));
+			}
 		}
 	}
 }
 
-bool FVirtualizationManager::TryCacheDataToBackend(IVirtualizationBackend& Backend, const FIoHash& Id, const FCompressedBuffer& Payload)
+bool FVirtualizationManager::TryCacheDataToBackend(IVirtualizationBackend& Backend, TArrayView<FPushRequest> Requests)
 {
-	COOK_STAT(FCookStats::FScopedStatsCounter Timer(Profiling::GetCacheStats(Backend)));
+	COOK_STAT(FCookStats::CallStats & Stats = Profiling::GetCacheStats(Backend));
+	COOK_STAT(FCookStats::FScopedStatsCounter Timer(Stats));
+	COOK_STAT(Timer.TrackCyclesOnly());
 	
-	FPushRequest Request(Id, Payload, FString());
-
-	if (Backend.PushData(MakeArrayView(&Request, 1)))
+	if (Backend.PushData(Requests))
 	{
-		if (Request.GetResult().WasPushed())
+#if ENABLE_COOK_STATS
+		Timer.AddHit(0);
+
+		const bool bIsInGameThread = IsInGameThread();
+
+		for (const FPushRequest& Request : Requests)
 		{
-			COOK_STAT(Timer.AddHit(Payload.GetCompressedSize()));
+			if (Request.GetResult().WasPushed())
+			{
+				Stats.Accumulate(FCookStats::CallStats::EHitOrMiss::Hit, FCookStats::CallStats::EStatType::Counter, 1l, bIsInGameThread);
+				Stats.Accumulate(FCookStats::CallStats::EHitOrMiss::Hit, FCookStats::CallStats::EStatType::Bytes, Request.GetPayload().GetCompressedSize(), bIsInGameThread);
+			}
 		}
+#endif // ENABLE_COOK_STATS
 
 		return true;
 	}
@@ -1622,20 +1766,23 @@ bool FVirtualizationManager::TryPushDataToBackend(IVirtualizationBackend& Backen
 	return bPushResult;
 }
 
-FCompressedBuffer FVirtualizationManager::PullDataFromAllBackends(const FIoHash& Id)
+void FVirtualizationManager::PullDataFromAllBackends(TArrayView<FPullRequest> Requests)
 {
 	if (ShouldDebugFailPulling())
 	{
-		UE_LOG(LogVirtualization, Verbose, TEXT("Debug miss chance (%.1f%%) invoked when pulling payload '%s'"), DebugValues.MissChance, *LexToString(Id));
-		return FCompressedBuffer();
+		UE_LOG(LogVirtualization, Verbose, TEXT("Debug miss chance (%.1f%%) invoked"), DebugValues.MissChance);
+		return;
 	}
 
+	FPullRequestCollection RequestsCollection(Requests);
+
 	for (IVirtualizationBackend* Backend : PullEnabledBackends)
 	{
-		// Skip if pulling has been disabled on this backend for debug purposes
+		check(Backend != nullptr);
+
 		if (Backend->IsOperationDebugDisabled(IVirtualizationBackend::EOperations::Pull))
 		{
-			UE_LOG(LogVirtualization, Verbose, TEXT("Pulling from backend '%s' is debug disabled for payload '%s'"), *Backend->GetDebugName(), *LexToString(Id));
+			UE_LOG(LogVirtualization, Verbose, TEXT("Pulling from backend '%s' is debug disabled"), *Backend->GetDebugName());
 			continue;
 		}
 
@@ -1645,38 +1792,45 @@ FCompressedBuffer FVirtualizationManager::PullDataFromAllBackends(const FIoHash&
 			continue;
 		}
 
-		FCompressedBuffer Payload = PullDataFromBackend(*Backend, Id);
+		PullDataFromBackend(*Backend, RequestsCollection.GetRequests());
 
-		if (Payload)
+		TArray<FPushRequest> PayloadsToCache = RequestsCollection.OnPullCompleted(*Backend, bEnableCacheAfterPull);
+		if (!PayloadsToCache.IsEmpty())
 		{
-			if (bEnableCacheAfterPull)
-			{
-				CachePayload(Id, Payload, Backend);
-			}
-
-			UE_LOG(LogVirtualization, VeryVerbose, TEXT("[%s] pulled payload '%s'"), *Backend->GetDebugName(), *LexToString(Id));
+			CachePayloads(PayloadsToCache, Backend);
+		}
 
-			return Payload;
+		if (!RequestsCollection.HasRemainingRequests())
+		{
+			break;
 		}
 	}
-
-	return FCompressedBuffer();
 }
 
-FCompressedBuffer FVirtualizationManager::PullDataFromBackend(IVirtualizationBackend& Backend, const FIoHash& Id)
+void FVirtualizationManager::PullDataFromBackend(IVirtualizationBackend& Backend, TArrayView<FPullRequest> Requests)
 {
-	COOK_STAT(FCookStats::FScopedStatsCounter Timer(Profiling::GetPullStats(Backend)));
-	FCompressedBuffer Payload = Backend.PullData(Id);
+	COOK_STAT(FCookStats::CallStats & Stats = Profiling::GetPullStats(Backend));
+	COOK_STAT(FCookStats::FScopedStatsCounter Timer(Stats));
+	COOK_STAT(Timer.TrackCyclesOnly());
+	
+	Backend.PullData(Requests);
 	
-	if (!Payload.IsNull())
+#if ENABLE_COOK_STATS
+	const bool bIsInGameThread = IsInGameThread();
+
+	for (const FPullRequest& Request : Requests)
 	{
-		COOK_STAT(Timer.AddHit(Payload.GetCompressedSize()));
+		Timer.AddHit(0);
+
+		if (Request.IsSuccess())
+		{
+			Stats.Accumulate(FCookStats::CallStats::EHitOrMiss::Hit, FCookStats::CallStats::EStatType::Counter, 1l, bIsInGameThread);
+			Stats.Accumulate(FCookStats::CallStats::EHitOrMiss::Hit, FCookStats::CallStats::EStatType::Bytes, Request.GetPayload().GetCompressedSize(), bIsInGameThread);
+		}
 	}
-	
-	return Payload;
+#endif //ENABLE_COOK_STATS
 }
 
-
 bool FVirtualizationManager::ShouldVirtualizeAsset(const UObject* OwnerObject) const
 {
 	if (OwnerObject == nullptr)
@@ -1814,6 +1968,15 @@ bool FVirtualizationManager::ShouldVirtualizeAsDefault() const
 	}
 }
 
+void FVirtualizationManager::BroadcastEvent(TConstArrayView<FPullRequest> Requests, ENotification Event)
+{
+	for (const FPullRequest& Request : Requests)
+	{
+		GetNotificationEvent().Broadcast(IVirtualizationSystem::PullEndedNotification, Request.GetIdentifier());
+	}
+}
+
 } // namespace UE::Virtualization
 
+#undef UE_INLINE_ALLOCATION_COUNT
 #undef LOCTEXT_NAMESPACE
\ No newline at end of file
diff --git a/Engine/Source/Developer/Virtualization/Private/VirtualizationManager.h b/Engine/Source/Developer/Virtualization/Private/VirtualizationManager.h
--- a/Engine/Source/Developer/Virtualization/Private/VirtualizationManager.h
+++ b/Engine/Source/Developer/Virtualization/Private/VirtualizationManager.h
@@ -83,6 +83,7 @@ namespace UE::Virtualization
 {
 class IVirtualizationBackend;
 class IVirtualizationBackendFactory;
+class FPullRequestCollection;
 
 /** The default mode of filtering to use with package paths that do not match entries in UVirtualizationFilterSettings */
 enum class EPackageFilterMode : uint8
@@ -120,8 +121,7 @@ private:
 	virtual bool AllowSubmitIfVirtualizationFailed() const override;
 	
 	virtual bool PushData(TArrayView<FPushRequest> Requests, EStorageType StorageType) override;
-
-	virtual FCompressedBuffer PullData(const FIoHash& Id) override;
+	virtual bool PullData(TArrayView<FPullRequest> Requests) override;
 
 	virtual EQueryResult QueryPayloadStatuses(TArrayView<const FIoHash> Ids, EStorageType StorageType, TArray<EPayloadStatus>& OutStatuses) override;
 
@@ -168,13 +168,13 @@ private:
 
 	void EnsureBackendConnections();
 
-	void CachePayload(const FIoHash& Id, const FCompressedBuffer& Payload, const IVirtualizationBackend* BackendSource);
+	void CachePayloads(TArrayView<FPushRequest> Requests, const IVirtualizationBackend* BackendSource);
 
-	bool TryCacheDataToBackend(IVirtualizationBackend& Backend, const FIoHash& Id, const FCompressedBuffer& Payload);
+	bool TryCacheDataToBackend(IVirtualizationBackend& Backend, TArrayView<FPushRequest> Requests);
 	bool TryPushDataToBackend(IVirtualizationBackend& Backend, TArrayView<FPushRequest> Requests);
 
-	FCompressedBuffer PullDataFromAllBackends(const FIoHash& Id);
-	FCompressedBuffer PullDataFromBackend(IVirtualizationBackend& Backend, const FIoHash& Id);
+	void PullDataFromAllBackends(TArrayView<FPullRequest> Requests);
+	void PullDataFromBackend(IVirtualizationBackend& Backend, TArrayView<FPullRequest> Requests);
 
 	bool ShouldVirtualizeAsset(const UObject* Owner) const;
 
@@ -204,6 +204,8 @@ private:
 
 	/** Determines if the default filtering behavior is to virtualize a payload or not */
 	bool ShouldVirtualizeAsDefault() const;
+
+	void BroadcastEvent(TConstArrayView<FPullRequest> Ids, ENotification Event);
 	
 private:
 	// The following members are set from the config file
diff --git a/Engine/Source/Developer/Virtualization/Private/VirtualizationSourceControlBackend.cpp b/Engine/Source/Developer/Virtualization/Private/VirtualizationSourceControlBackend.cpp
--- a/Engine/Source/Developer/Virtualization/Private/VirtualizationSourceControlBackend.cpp
+++ b/Engine/Source/Developer/Virtualization/Private/VirtualizationSourceControlBackend.cpp
@@ -348,19 +348,87 @@ IVirtualizationBackend::EConnectionStatus FSourceControlBackend::OnConnect()
 	return IVirtualizationBackend::EConnectionStatus::Connected;
 }
 
-FCompressedBuffer FSourceControlBackend::PullData(const FIoHash& Id)
+bool FSourceControlBackend::PullData(TArrayView<FPullRequest> Requests)
 {
 	TRACE_CPUPROFILER_EVENT_SCOPE(FSourceControlBackend::PullData);
 
-	TStringBuilder<512> DepotPath;
-	CreateDepotPath(Id, DepotPath);
+#if 0
+	TRACE_CPUPROFILER_EVENT_SCOPE(FSourceControlBackend::PullData);
+
+	for (FPullRequest& Request : Requests)
+	{
+		TStringBuilder<512> DepotPath;
+		CreateDepotPath(Request.Identifier, DepotPath);
+
+		// TODO: When multiple threads are blocked waiting on this we could gather X payloads together and make a single
+		// batch request on the same connection, which should be a lot faster with less overhead.
+		// Although ideally this backend will not get hit very often.
+		FSemaphoreScopeLock _(ConcurrentConnectionLimit.Get());
+
+		UE_LOG(LogVirtualization, Verbose, TEXT("[%s] Attempting to pull '%s' from source control"), *GetDebugName(), *DepotPath);
+
+		int32 Retries = 0;
+
+		while (Retries < RetryCount)
+		{
+			// Only warn if the backend is configured to retry
+			if (Retries != 0)
+			{
+				UE_LOG(LogVirtualization, Warning, TEXT("[%s] Failed to download '%s' retrying (%d/%d) in %dms..."), *GetDebugName(), DepotPath.ToString(), Retries, RetryCount, RetryWaitTimeMS);
+				FPlatformProcess::SleepNoStats(RetryWaitTimeMS * 0.001f);
+			}
+
+#if IS_SOURCE_CONTROL_THREAD_SAFE
+			TSharedRef<FDownloadFile, ESPMode::ThreadSafe> DownloadCommand = ISourceControlOperation::Create<FDownloadFile>(FDownloadFile::EVerbosity::None);
+			if (SCCProvider->Execute(DownloadCommand, DepotPath.ToString(), EConcurrency::Synchronous) == ECommandResult::Succeeded)
+			{
+				// The payload was created by FCompressedBuffer::Compress so we can return it as a FCompressedBuffer.
+				FSharedBuffer Buffer = DownloadCommand->GetFileData(DepotPath);
+				return FCompressedBuffer::FromCompressed(Buffer);
+			}
+#else
+			TSharedRef<FDownloadFile> DownloadCommand = ISourceControlOperation::Create<FDownloadFile>(FDownloadFile::EVerbosity::None);
+			if (SCCProvider->TryToDownloadFileFromBackgroundThread(DownloadCommand, DepotPath.ToString()))
+			{
+				// The payload was created by FCompressedBuffer::Compress so we can return it as a FCompressedBuffer.
+				Request.Status = FPullRequest::EStatus::Success;
+				Request.Payload = FCompressedBuffer::FromCompressed(DownloadCommand->GetFileData(DepotPath));
+
+				break;
+			}
+#endif
+
+			// If this was the first try then check to see if the error being returns is that the file does not exist
+			// in the depot. If it does not exist then there is no point in us retrying and we can error out at this point.
+			if (Retries == 0 && IsDepotFileMissing(DownloadCommand->GetResultInfo()))
+			{
+				break;
+			}
+
+			Retries++;
+		}
+	}
+
+	return true;
+
+#else
+	TArray<FString> DepotPaths;
+	DepotPaths.Reserve(Requests.Num());
+
+	for (FPullRequest& Request : Requests)
+	{
+		TStringBuilder<512> DepotPath;
+		CreateDepotPath(Request.GetIdentifier(), DepotPath);
+
+		DepotPaths.Add(DepotPath.ToString());
+	}
 
 	// TODO: When multiple threads are blocked waiting on this we could gather X payloads together and make a single
 	// batch request on the same connection, which should be a lot faster with less overhead.
 	// Although ideally this backend will not get hit very often.
 	FSemaphoreScopeLock _(ConcurrentConnectionLimit.Get());
 
-	UE_LOG(LogVirtualization, Verbose, TEXT("[%s] Attempting to pull '%s' from source control"), *GetDebugName(), *DepotPath);
+//	UE_LOG(LogVirtualization, Verbose, TEXT("[%s] Attempting to pull '%s' from source control"), *GetDebugName(), *DepotPath);
 
 	int32 Retries = 0;
 
@@ -369,7 +437,7 @@ FCompressedBuffer FSourceControlBackend::PullData(const FIoHash& Id)
 		// Only warn if the backend is configured to retry
 		if (Retries != 0)
 		{
-			UE_LOG(LogVirtualization, Warning, TEXT("[%s] Failed to download '%s' retrying (%d/%d) in %dms..."), *GetDebugName(), DepotPath.ToString(), Retries, RetryCount, RetryWaitTimeMS);
+		//	UE_LOG(LogVirtualization, Warning, TEXT("[%s] Failed to download '%s' retrying (%d/%d) in %dms..."), *GetDebugName(), DepotPath.ToString(), Retries, RetryCount, RetryWaitTimeMS);
 			FPlatformProcess::SleepNoStats(RetryWaitTimeMS * 0.001f);
 		}
 
@@ -379,15 +447,23 @@ FCompressedBuffer FSourceControlBackend::PullData(const FIoHash& Id)
 		{
 			// The payload was created by FCompressedBuffer::Compress so we can return it as a FCompressedBuffer.
 			FSharedBuffer Buffer = DownloadCommand->GetFileData(DepotPath);
-			return FCompressedBuffer::FromCompressed(Buffer);
+
+			Request.Payload = FCompressedBuffer::FromCompressed(Buffer);
+			Request.Status = FPullRequest::EStatus::Success;
 		}
 #else
 		TSharedRef<FDownloadFile> DownloadCommand = ISourceControlOperation::Create<FDownloadFile>(FDownloadFile::EVerbosity::None);
-		if (SCCProvider->TryToDownloadFileFromBackgroundThread(DownloadCommand, DepotPath.ToString()))
+		if (SCCProvider->TryToDownloadFileFromBackgroundThread(DownloadCommand, DepotPaths))
 		{
-			// The payload was created by FCompressedBuffer::Compress so we can return it as a FCompressedBuffer.
-			FSharedBuffer Buffer = DownloadCommand->GetFileData(DepotPath);
-			return FCompressedBuffer::FromCompressed(Buffer);
+			for(int32 Index = 0; Index < Requests.Num(); ++Index)
+			{
+				// The payload was created by FCompressedBuffer::Compress so we can return it as a FCompressedBuffer.
+				FSharedBuffer Buffer = DownloadCommand->GetFileData(DepotPaths[Index]);
+
+				Requests[Index].SetPayload(FCompressedBuffer::FromCompressed(Buffer));
+			}
+
+			break;
 		}
 #endif
 
@@ -395,13 +471,16 @@ FCompressedBuffer FSourceControlBackend::PullData(const FIoHash& Id)
 		// in the depot. If it does not exist then there is no point in us retrying and we can error out at this point.
 		if (Retries == 0 && IsDepotFileMissing(DownloadCommand->GetResultInfo()))
 		{
-			return FCompressedBuffer();
+			break;
 		}
-		
+
 		Retries++;
 	}
 
-	return FCompressedBuffer();
+	// TODO -  Should we error here if the payload failed? 
+
+	return true;
+#endif
 }
 
 bool FSourceControlBackend::DoesPayloadExist(const FIoHash& Id)
diff --git a/Engine/Source/Developer/Virtualization/Private/VirtualizationSourceControlBackend.h b/Engine/Source/Developer/Virtualization/Private/VirtualizationSourceControlBackend.h
--- a/Engine/Source/Developer/Virtualization/Private/VirtualizationSourceControlBackend.h
+++ b/Engine/Source/Developer/Virtualization/Private/VirtualizationSourceControlBackend.h
@@ -83,8 +83,7 @@ private:
 	virtual EConnectionStatus OnConnect() override;
 	
 	virtual bool PushData(TArrayView<FPushRequest> Requests) override;
-
-	virtual FCompressedBuffer PullData(const FIoHash& Id) override;
+	virtual bool PullData(TArrayView<FPullRequest> Requests) override;
 
 	virtual bool DoesPayloadExist(const FIoHash& Id) override;
 	
diff --git a/Engine/Source/Developer/Virtualization/Public/IVirtualizationBackend.h b/Engine/Source/Developer/Virtualization/Public/IVirtualizationBackend.h
--- a/Engine/Source/Developer/Virtualization/Public/IVirtualizationBackend.h
+++ b/Engine/Source/Developer/Virtualization/Public/IVirtualizationBackend.h
@@ -106,17 +106,26 @@ public:
 	virtual bool PushData(TArrayView<FPushRequest> Requests) = 0;
 
 	/** 
-	 * The backend will attempt to retrieve the given payload by what ever method the backend uses.
-	 * NOTE: It is assumed that the virtualization manager will validate the returned payload to 
-	 * make sure that it matches the requested id so there is no need for each backend to do this/
+	 * The backend will attempt to retrieve the given payloads by what ever method the backend uses.
 	 * 
-	 * @param Id	The Id of a payload to try and pull from the backend.
+	 * It should be assumed that the list of requests will not contain any duplicate or invalid
+	 * payload identifiers so there is no need for each backend to perform validation.
 	 * 
-	 * @return		A valid FCompressedBuffer containing the payload if the pull
-	 *				operation succeeded and a null FCompressedBuffer
-	 *				if it did not.
+	 * It should be assumed that the caller will validate all payloads that are successfully pulled
+	 * to make sure that match the requested payload identifiers  so there is no need for each 
+	 * backend to do this.
+	 * 
+	 * @param	Requests	An array of payload pull requests. @see FPullRequest
+	 * 
+	 * @return				True if no errors were encountered while pulling, otherwise false.
+	 *						Note that returning true does not mean that all of the payloads were
+	 *						found as a missing payload should not be considered an error condition.
 	 */
-	virtual FCompressedBuffer PullData(const FIoHash& Id) = 0;
+
+	// Assume that the array only has unique requests
+	// Will set request error value if there is a problem with loading it
+	// Returns false on critical error, otherwise true
+	virtual bool PullData(TArrayView<FPullRequest> Requests) = 0;
 	
 	/**
 	 * Checks if a payload exists in the backends storage.
diff --git a/Engine/Source/Runtime/Core/Private/Virtualization/VirtualizationSystem.cpp b/Engine/Source/Runtime/Core/Private/Virtualization/VirtualizationSystem.cpp
--- a/Engine/Source/Runtime/Core/Private/Virtualization/VirtualizationSystem.cpp
+++ b/Engine/Source/Runtime/Core/Private/Virtualization/VirtualizationSystem.cpp
@@ -92,13 +92,16 @@ public:
 		return false;
 	}
 
-	virtual FCompressedBuffer PullData(const FIoHash& Id) override
+	virtual bool PullData(TArrayView<FPullRequest> Requests) override
 	{
 #if UE_LOG_ON_NULLSYSTEM_USE
-		UE_LOG(LogVirtualization, UE_NULLSYSTEM_SEVERITY, TEXT("Cannot pull payload '%s' as the virtualization system is disabled"), *LexToString(Id));
+		for (const FPullRequest& Request : Requests)
+		{
+			UE_LOG(LogVirtualization, UE_NULLSYSTEM_SEVERITY, TEXT("Cannot pull payload '%s' as the virtualization system is disabled"), *LexToString(Request.GetIdentifier()));
+		}
 #endif //UE_LOG_ON_NULLSYSTEM_USE
 
-		return FCompressedBuffer();
+		return false;
 	}
 
 	virtual EQueryResult QueryPayloadStatuses(TArrayView<const FIoHash> Ids, EStorageType StorageType, TArray<EPayloadStatus>& OutStatuses) override
diff --git a/Engine/Source/Runtime/Core/Public/Virtualization/VirtualizationSystem.h b/Engine/Source/Runtime/Core/Public/Virtualization/VirtualizationSystem.h
--- a/Engine/Source/Runtime/Core/Public/Virtualization/VirtualizationSystem.h
+++ b/Engine/Source/Runtime/Core/Public/Virtualization/VirtualizationSystem.h
@@ -372,6 +372,63 @@ private:
 	FPushResult Result;
 };
 
+/** Data structure representing a payload pull request */
+struct FPullRequest
+{
+private:
+	enum class EStatus : int8
+	{
+		Error = -1,
+		Pending = 0,
+		Success = 1
+	};
+
+public:
+	FPullRequest(const FIoHash& InIdentifier)
+		: Identifier(InIdentifier)
+	{
+
+	}
+
+	const FIoHash& GetIdentifier() const
+	{
+		return Identifier;
+	}
+
+	const FCompressedBuffer& GetPayload() const
+	{
+		return Payload;
+	}
+
+	bool IsSuccess() const
+	{
+		return Status == EStatus::Success;
+	}
+
+	void SetError()
+	{
+		Status = EStatus::Error;
+		Payload.Reset();
+	}
+
+	void SetPayload(const FCompressedBuffer& InPayload)
+	{
+		Payload = InPayload;
+		Status = Payload ? EStatus::Success : EStatus::Error;
+	}
+
+	void SetPayload(FCompressedBuffer&& InPayload)
+	{
+		Payload = MoveTemp(InPayload);
+		Status = Payload ? EStatus::Success : EStatus::Error;
+	}
+
+private:
+	FIoHash Identifier;
+	FCompressedBuffer Payload;
+	EStatus Status = EStatus::Pending;
+};
+
 /** 
  * The set of parameters to be used when initializing the virtualization system. The 
  * members must remain valid for the duration of the call to ::Initialize. It is not
@@ -561,14 +618,45 @@ public:
 	}
 
 	/**
-	 * Pull a payload from the virtualization backends.
+	 * Pull a number of payloads from the virtualization backends.
+	 *
+	 * @param	Requests	An array of payload pull requests. @see FPullRequest
+	 * @return				True if the requests succeeded, false if one or more requests failed
+	 */
+	virtual bool PullData(TArrayView<FPullRequest> Requests) = 0;
+
+	/** 
+	 * Pull a single payload from the virtualization backends.
+	 * 
+	 * @param Id	The hash of the payload to pull
+	 * @return		A valid buffer representing the payload if the pull was successful. 
+	 *				An invalid buffer if the pull failed.
+	 */
+	FCompressedBuffer PullData(const FIoHash& Id)
+	{
+		FPullRequest Request(Id);
+
+		if (PullData(MakeArrayView(&Request, 1)))
+		{
+			return Request.GetPayload();
+		}
+		else
+		{
+			return FCompressedBuffer();
+		}
+	}
+
+	/**
+	 * Pull a single payload from the virtualization backends.
 	 *
-	 * @param	Id The identifier of the payload being pulled.
-	 * @return	The payload in the form of a FCompressedBuffer. No decompression will be applied to the payload, it is 
-	 *			up to the caller if they want to retain the payload in compressed or uncompressed format.  If no
-	 *			backend contained the payload then an empty invalid FCompressedBuffer will be returned.
+	 * @param Request	A payload pull request. @see FPullRequest
+	 * @return			True if the pull failed and the request will return a valid payload.
+	 *					False if the pull failed and the request will return an invalid payload
 	 */
-	virtual FCompressedBuffer PullData(const FIoHash& Id) = 0;
+	bool PullData(FPullRequest Request)
+	{
+		return PullData(MakeArrayView(&Request, 1));
+	}
 
 	/**
 	 * Query if a number of payloads exist or not in the given storage type. 
