Fix some broken VA debug commandline options

#rnx

- Work on UE-165308 broke the following commandlines as they were accidently converted to cmdline switches rather than values
-- VABackendGraph
-- VAMissBackends
-- VAMissChance
- ::IsCmdLineSet has now been split into two utility wrappers, one for params (bools) and one for values.

diff --git a/Engine/Source/Developer/Virtualization/Private/VirtualizationManager.cpp b/Engine/Source/Developer/Virtualization/Private/VirtualizationManager.cpp
--- a/Engine/Source/Developer/Virtualization/Private/VirtualizationManager.cpp
+++ b/Engine/Source/Developer/Virtualization/Private/VirtualizationManager.cpp
@@ -221,11 +221,11 @@ bool LexTryParseString(EPackageFilterMode& OutValue, FStringView Buffer)
 	return false;
 }
 
-/** 
- * Utility to check if either cmd is present in the command line. Useful when transitioning from one
- * command line to another. 
+/**
+ * Utility to check if a cmdline switch is present in the commandline under either of two names.
+ * Useful when transitioning from one command line to another.
  */
-static bool IsCmdLineSet(const TCHAR* Cmd, const TCHAR* AlternativeCmd = nullptr)
+static bool IsCmdLineParamSet(const TCHAR* Cmd, const TCHAR* AlternativeCmd)
 {
 	const TCHAR* CmdLine = FCommandLine::Get();
 
@@ -234,7 +234,29 @@ static bool IsCmdLineSet(const TCHAR* Cmd, const TCHAR* AlternativeCmd = nullptr
 		return true;
 	}
 
-	if (AlternativeCmd != nullptr && FParse::Param(CmdLine, AlternativeCmd))
+	if (FParse::Param(CmdLine, AlternativeCmd))
+	{
+		return true;
+	}
+
+	return false;
+}
+
+/**
+ * Utility to check if a cmdline value is present in the commandline under either of two names.
+ * Useful when transitioning from one command line to another.
+ */
+template<typename T>
+static bool IsCmdLineValueSet(const TCHAR* Cmd, const TCHAR* AlternativeCmd, T& OutValue)
+{
+	const TCHAR* CmdLine = FCommandLine::Get();
+
+	if (FParse::Value(CmdLine, Cmd, OutValue))
+	{
+		return true;
+	}
+
+	if (FParse::Value(CmdLine, AlternativeCmd, OutValue))
 	{
 		return true;
 	}
@@ -1126,13 +1148,13 @@ void FVirtualizationManager::ApplySettingsFromConfigFiles(const FConfigFile& Con
 
 void FVirtualizationManager::ApplySettingsFromFromCmdline()
 {
-	if (!bLazyInitConnections && IsCmdLineSet(TEXT("VALazyInitConnections"), TEXT("VA-LazyInitConnections")))
+	if (!bLazyInitConnections && IsCmdLineParamSet(TEXT("VALazyInitConnections"), TEXT("VA-LazyInitConnections")))
 	{
 		bLazyInitConnections = true;
 		UE_LOG(LogVirtualization, Display, TEXT("Cmdline has set the virtualization system backends to lazy init their connections"));
 	}
 
-	if (bAllowPackageVirtualization && IsCmdLineSet(TEXT("VASkipPkgVirtualization"), TEXT("VA-SkipPkgVirtualization")))
+	if (bAllowPackageVirtualization && IsCmdLineParamSet(TEXT("VASkipPkgVirtualization"), TEXT("VA-SkipPkgVirtualization")))
 	{
 		bAllowPackageVirtualization = false;
 		UE_LOG(LogVirtualization, Warning, TEXT("The virtualization process has been disabled via the command line"));
@@ -1150,27 +1172,27 @@ void FVirtualizationManager::ApplySettingsFromCVar()
 
 void FVirtualizationManager::ApplyDebugSettingsFromFromCmdline()
 {
-	if (IsCmdLineSet(TEXT("VASingleThreaded"), TEXT("VA-SingleThreaded")))
+	if (IsCmdLineParamSet(TEXT("VASingleThreaded"), TEXT("VA-SingleThreaded")))
 	{
 		DebugValues.bSingleThreaded = true;
 		UE_LOG(LogVirtualization, Warning, TEXT("Cmdline has set the virtualization system to run single threaded"));
 	}
 
-	if (IsCmdLineSet(TEXT("VAValidatePushes"), TEXT("VA-ValidatePushes")))
+	if (IsCmdLineParamSet(TEXT("VAValidatePushes"), TEXT("VA-ValidatePushes")))
 	{
 		DebugValues.bValidateAfterPush = true;
 		UE_LOG(LogVirtualization, Warning, TEXT("Cmdline has set the virtualization system to pull each payload after pushing to either local or persistent storage"));
 	}
 
 	FString CmdlineGraphName;
-	if (IsCmdLineSet(TEXT("VABackendGraph"), TEXT("VA-BackendGraph")))
+	if (IsCmdLineValueSet(TEXT("-VABackendGraph="), TEXT("-VA-BackendGraph="), CmdlineGraphName))
 	{
 		UE_LOG(LogVirtualization, Display, TEXT("Backend graph overriden from the cmdline: '%s'"), *CmdlineGraphName);
 		BackendGraphName = CmdlineGraphName;
 	}
 
 	FString MissOptions;
-	if (IsCmdLineSet(TEXT("VAMissBackends"), TEXT("VA-MissBackends")))
+	if (IsCmdLineValueSet(TEXT("-VAMissBackends="), TEXT("-VA-MissBackends="), MissOptions))
 	{
 		MissOptions.ParseIntoArray(DebugValues.MissBackends, TEXT("+"), true);
 
@@ -1182,7 +1204,7 @@ void FVirtualizationManager::ApplyDebugSettingsFromFromCmdline()
 	}
 
 	DebugValues.MissChance = 0.0f;
-	if (IsCmdLineSet(TEXT("VAMissChance"), TEXT("VA-MissChance")))
+	if (IsCmdLineValueSet(TEXT("-VAMissChance="), TEXT("-VA-MissChance="), DebugValues.MissChance))
 	{
 		DebugValues.MissChance = FMath::Clamp(DebugValues.MissChance, 0.0f, 100.0f);
 
