Allow VA backends to be marked as read only to prevent new payloads being pushed to them.

#rnx

- IVirtualizationBackend can now have its operations disabled after it is created.
-- This means that a backend that only supported pushing and not pulling would end up with no valid operations if marked as read only. This will not assert because we do this after the backend is created and we only check on creation.

diff --git a/Engine/Source/Developer/Virtualization/Private/VirtualizationManager.cpp b/Engine/Source/Developer/Virtualization/Private/VirtualizationManager.cpp
--- a/Engine/Source/Developer/Virtualization/Private/VirtualizationManager.cpp
+++ b/Engine/Source/Developer/Virtualization/Private/VirtualizationManager.cpp
@@ -1603,11 +1603,19 @@ bool FVirtualizationManager::CreateBackend(const FConfigFile& ConfigFile, const
 			{
 				UE_LOG(LogVirtualization, Fatal, TEXT("IVirtualizationBackendFactory '%s' failed to create an instance!"), *Factory->GetName().ToString());
 				return false;
-
 			}
 
 			if (Backend->Initialize(Cmdine))
 			{
+				// The read only flag can be applied to any backend so we check for it and apply it at this point
+				bool bReadOnly = false;
+				FParse::Bool(*Cmdine, TEXT("ReadOnly="), bReadOnly);
+
+				if (bReadOnly && Backend->DisableOperation(IVirtualizationBackend::EOperations::Push))
+				{
+					UE_LOG(LogVirtualization, Display, TEXT("The backend '%s' was set to readonly by the config file!"), *Backend->GetDebugName());
+				}
+
 				AddBackend(MoveTemp(Backend), PushArray);
 			}
 			else
diff --git a/Engine/Source/Developer/Virtualization/Public/IVirtualizationBackend.h b/Engine/Source/Developer/Virtualization/Public/IVirtualizationBackend.h
--- a/Engine/Source/Developer/Virtualization/Public/IVirtualizationBackend.h
+++ b/Engine/Source/Developer/Virtualization/Public/IVirtualizationBackend.h
@@ -162,6 +162,22 @@ public:
 
 		return true;
 	}
+
+	/**
+	 * Disables operations on the backend
+	 * @param Operation	The operation(s) to be disabled
+	 * 
+	 * @return	True if at least one previously enabled operation was disabled by the call
+	 *			False if there was no change.
+	 */
+	bool DisableOperation(EOperations Operation)
+	{
+		EOperations OriginalValue = SupportedOperations;
+
+		EnumRemoveFlags(SupportedOperations, Operation);
+
+		return OriginalValue != SupportedOperations;
+	}
 	
 	/** 
 	 * Returns true if the given operation is supported, this is set when the backend is created
