Allow the VA file system backend root path to be overriden by an environment variable, allowing users to choose where they want the data to be stored.

#rnx

diff --git a/Engine/Source/Developer/Virtualization/Private/VirtualizationFileBackend.cpp b/Engine/Source/Developer/Virtualization/Private/VirtualizationFileBackend.cpp
--- a/Engine/Source/Developer/Virtualization/Private/VirtualizationFileBackend.cpp
+++ b/Engine/Source/Developer/Virtualization/Private/VirtualizationFileBackend.cpp
@@ -24,6 +24,17 @@ bool FFileSystemBackend::Initialize(const FString& ConfigEntry)
 		return false;
 	}
 
+	FString EnvOverrideName;
+	if (FParse::Value(*ConfigEntry, TEXT("EnvPathOverride="), EnvOverrideName))
+	{
+		FString OverridePath = FPlatformMisc::GetEnvironmentVariable(*EnvOverrideName);
+		if (!OverridePath.IsEmpty())
+		{
+			UE_LOG(LogVirtualization, Log, TEXT("[%s] Overriding path with envvar '%s'"), *GetDebugName(), *EnvOverrideName);
+			RootDirectory = OverridePath;
+		}
+	}
+
 	FPaths::NormalizeDirectoryName(RootDirectory);
 
 	if (RootDirectory.IsEmpty())
diff --git a/Engine/Source/Developer/Virtualization/Private/VirtualizationFileBackend.h b/Engine/Source/Developer/Virtualization/Private/VirtualizationFileBackend.h
--- a/Engine/Source/Developer/Virtualization/Private/VirtualizationFileBackend.h
+++ b/Engine/Source/Developer/Virtualization/Private/VirtualizationFileBackend.h
@@ -27,6 +27,9 @@ namespace UE::Virtualization
  * RetryWaitTime:	The length of time the process should wait between each read attempt in milliseconds.
  *					Remember that the max length of time that the process can stall attempting to read
  *					a payload file is RetryCount * RetryWaitTime; [Default=100ms]
+ * EnvPathOverride: The name of the environment variable that the user can set to override 'Path'. This 
+ *					should only ever be set if you wish for the users to be able to choose where the 
+ *					payloads are being stored and do not rely on 'Path' always being used. [Default=""]
  */
 class FFileSystemBackend final : public IVirtualizationBackend
 {
