Removed AnalyticsEventAttribute.h from the ZenServerInterface and VirtualizationManager header

diff --git a/Engine/Source/Developer/Virtualization/Private/VirtualizationManager.cpp b/Engine/Source/Developer/Virtualization/Private/VirtualizationManager.cpp
--- a/Engine/Source/Developer/Virtualization/Private/VirtualizationManager.cpp
+++ b/Engine/Source/Developer/Virtualization/Private/VirtualizationManager.cpp
@@ -19,6 +19,7 @@
 #include "ProfilingDebugging/CookStats.h"
 #include "VirtualizationFilterSettings.h"
 #include "Analytics.h"
+#include "AnalyticsEventAttribute.h"
 
 #define LOCTEXT_NAMESPACE "Virtualization"
 
diff --git a/Engine/Source/Developer/Virtualization/Private/VirtualizationManager.h b/Engine/Source/Developer/Virtualization/Private/VirtualizationManager.h
--- a/Engine/Source/Developer/Virtualization/Private/VirtualizationManager.h
+++ b/Engine/Source/Developer/Virtualization/Private/VirtualizationManager.h
@@ -6,13 +6,13 @@
 #include "HAL/CriticalSection.h"
 #include "Logging/LogMacros.h"
 #include "Templates/UniquePtr.h"
-#include "AnalyticsEventAttribute.h"
 
 
 #include "Virtualization/VirtualizationSystem.h"
 
 class IConsoleObject;
 class FOutputDevice;
+struct FAnalyticsEventAttribute;
 
 /**
  * Configuring the backend hierarchy
diff --git a/Engine/Source/Developer/Zen/Internal/ZenServerInterface.cpp b/Engine/Source/Developer/Zen/Internal/ZenServerInterface.cpp
--- a/Engine/Source/Developer/Zen/Internal/ZenServerInterface.cpp
+++ b/Engine/Source/Developer/Zen/Internal/ZenServerInterface.cpp
@@ -27,6 +27,7 @@
 #include "String/LexFromString.h"
 #include "SocketSubsystem.h"
 #include "Analytics.h"
+#include "AnalyticsEventAttribute.h"
 
 #if PLATFORM_WINDOWS
 #	include "Windows/AllowWindowsPlatformTypes.h"
diff --git a/Engine/Source/Developer/Zen/Public/ZenServerInterface.h b/Engine/Source/Developer/Zen/Public/ZenServerInterface.h
--- a/Engine/Source/Developer/Zen/Public/ZenServerInterface.h
+++ b/Engine/Source/Developer/Zen/Public/ZenServerInterface.h
@@ -15,7 +15,6 @@
 #include "Serialization/JsonWriter.h"
 #include "Templates/PimplPtr.h"
 #include "Templates/UniquePtr.h"
-#include "AnalyticsEventAttribute.h"
 #include "ZenGlobals.h"
 
 #if UE_WITH_ZEN
@@ -24,6 +23,8 @@
 
 #define UE_API ZEN_API
 
+struct FAnalyticsEventAttribute;
+
 namespace UE::Zen
 {
 
