Fix the VA filesystem backend from failing to serialize pushed payloads correctly.

#rnx

- The original for loop was added at a time where FCompressedBuffer did not support serialization directly, at that time the code worked.
- Now it looks like the FCompressedBuffer returned by the call to ::GetPayload is going out of scope early, which means that the FSharedBuffer ends up zero length (and quite possibly unsafe to use)
- However nowadays FCompressedBuffer has both a operator<< and direct ::Load/::Save methods so we can remove the for loop and simplify the code.

diff --git a/Engine/Source/Developer/Virtualization/Private/VirtualizationFileBackend.cpp b/Engine/Source/Developer/Virtualization/Private/VirtualizationFileBackend.cpp
--- a/Engine/Source/Developer/Virtualization/Private/VirtualizationFileBackend.cpp
+++ b/Engine/Source/Developer/Virtualization/Private/VirtualizationFileBackend.cpp
@@ -113,10 +113,9 @@ bool FFileSystemBackend::PushData(TArrayView<FPushRequest> Requests)
 			continue;
 		}
 
-		for (const FSharedBuffer& Buffer : Request.GetPayload().GetCompressed().GetSegments())
 		{
-			// Const cast because FArchive requires a non-const pointer!
-			FileAr->Serialize(const_cast<void*>(Buffer.GetData()), static_cast<int64>(Buffer.GetSize()));
+			FCompressedBuffer Payload = Request.GetPayload();
+			*FileAr << Payload;
 		}
 
 		if (!FileAr->Close())
