Virtualizing an asset will no longer display an error message if the project has no caching storage backends set up.

#rnx

- It is possible that a project might want to have persistent storage enabled but no caching storage, it is probably a bad idea but it should work. However before this fix doing so would cause an error to be logged during the virtualization of a package and fail the process.
- During the virtualization process we now check if pushing is enabled for cache storage before trying to push there, if not then we log that the phase is being skipped.
- Even though the above fixes the problem we would still get odd results when pushing payloads to cached storage without any backends enabled.
-- We now check if the pushing process is disabled before processing requests, so we can fail faster. There is now also a proper error value for this that can be stored in each request (FPushRequest) where as previously they'd all still display "pending"
-- We also check to see if the requested push type has any associated backend before processing requests and have a specific error value for this. Note that we consider
- We also now force the visibility of the slow task progress bar, as otherwise we might not end up displaying the virtualization progress to the progress dialog (this is a work around to a long standing problem with the slow task system due to the time limit associated with updating the dialog via EnterProgressFrame, this needs to be fixed elsewhere)

diff --git a/Engine/Source/Developer/Virtualization/Private/PackageVirtualizationProcess.cpp b/Engine/Source/Developer/Virtualization/Private/PackageVirtualizationProcess.cpp
--- a/Engine/Source/Developer/Virtualization/Private/PackageVirtualizationProcess.cpp
+++ b/Engine/Source/Developer/Virtualization/Private/PackageVirtualizationProcess.cpp
@@ -163,6 +163,8 @@ void VirtualizePackages(TConstArrayView<FString> PackagePaths, TArray<FText>& Ou
 	const double StartTime = FPlatformTime::Seconds();
 
 	FScopedSlowTask Progress(5.0f, LOCTEXT("Virtualization_Task", "Virtualizing Assets..."));
+	// Force the task to be visible otherwise it might not be shown if the initial progress frames are too fast
+	Progress.Visibility = ESlowTaskVisibility::ForceVisible;
 	Progress.MakeDialog();
 
 	// Other systems may have added errors to this array, we need to check so later we can determine if this function added any additional errors.
@@ -266,9 +268,9 @@ void VirtualizePackages(TConstArrayView<FString> PackagePaths, TArray<FText>& Ou
 	// TODO: We should be able to do both Cache and Persistent pushes in the same call
 
 	// Push payloads to cache storage
+	Progress.EnterProgressFrame(1.0f);
+	if(System.IsPushingEnabled(EStorageType::Cache))
 	{
-		Progress.EnterProgressFrame(1.0f);
-		
 		if (!System.PushData(PayloadsToSubmit, EStorageType::Cache))
 		{
 			// Caching is not critical to the process so we only warn if it fails
@@ -286,6 +288,10 @@ void VirtualizePackages(TConstArrayView<FString> PackagePaths, TArray<FText>& Ou
 		}
 		UE_LOG(LogVirtualization, Display, TEXT("Pushed %" INT64_FMT " payload(s) to cached storage"), TotalPayloadsCached);
 	}
+	else
+	{
+		UE_LOG(LogVirtualization, Display, TEXT("Pushing payload(s) to cached storage is disbled, skipping"));
+	}
 	
 	// Push payloads to persistent storage
 	{
diff --git a/Engine/Source/Developer/Virtualization/Private/VirtualizationManager.cpp b/Engine/Source/Developer/Virtualization/Private/VirtualizationManager.cpp
--- a/Engine/Source/Developer/Virtualization/Private/VirtualizationManager.cpp
+++ b/Engine/Source/Developer/Virtualization/Private/VirtualizationManager.cpp
@@ -266,6 +266,15 @@ static bool IsCmdLineValueSet(const TCHAR* Cmd, const TCHAR* AlternativeCmd, T&
 	return false;
 }
 
+/** Utility to set the same FPushResult on many requests at once */
+static void SetPushRequestsResult(TArrayView<FPushRequest> Requests, FPushResult Result)
+{
+	for (FPushRequest& Request : Requests)
+	{
+		Request.SetResult(Result);
+	}
+}
+
 /* Utility function for building up a lookup table of all available IBackendFactory interfaces*/
 FVirtualizationManager::FRegistedFactories FindBackendFactories()
 {
@@ -582,11 +591,27 @@ bool FVirtualizationManager::PushData(TArrayView<FPushRequest> Requests, EStorag
 {
 	TRACE_CPUPROFILER_EVENT_SCOPE(FVirtualizationManager::PushData);
 
+	// No requests always counts as success
 	if (Requests.IsEmpty())
 	{
 		return true;
 	}
 
+	// Attempting to virtualize when the system or process is disabled always counts as a failure
+	if (!IsEnabled() || !bAllowPackageVirtualization)
+	{
+		SetPushRequestsResult(Requests, FPushResult::GetAsProcessDisabled());
+		return false;
+	}
+
+	FBackendArray& Backends = IsCacheType(StorageType) ? CacheStorageBackends : PersistentStorageBackends;
+
+	if (Backends.IsEmpty())
+	{
+		SetPushRequestsResult(Requests, FPushResult::GetAsNoBackend());
+		return false;
+	}
+
 	TArray<FPushRequest> ValidatedRequests;
 	ValidatedRequests.Reserve(Requests.Num());
 
@@ -637,12 +662,6 @@ bool FVirtualizationManager::PushData(TArrayView<FPushRequest> Requests, EStorag
 		return true;
 	}
 
-	// Early out if there are no backends
-	if (!IsEnabled() || !bAllowPackageVirtualization)
-	{
-		return false;
-	}
-
 	EnsureBackendConnections();
 
 	FConditionalScopeLock _(&DebugValues.ForceSingleThreadedCS, DebugValues.bSingleThreaded);
@@ -653,8 +672,7 @@ bool FVirtualizationManager::PushData(TArrayView<FPushRequest> Requests, EStorag
 
 	int32 ErrorCount = 0;
 	bool bWasPayloadPushed = false;
-	FBackendArray& Backends = IsCacheType(StorageType) ? CacheStorageBackends : PersistentStorageBackends;
-
+	
 	for (IVirtualizationBackend* Backend : Backends)
 	{
 		if (Backend->GetConnectionStatus() != IVirtualizationBackend::EConnectionStatus::Connected)
diff --git a/Engine/Source/Runtime/Core/Public/Virtualization/VirtualizationSystem.h b/Engine/Source/Runtime/Core/Public/Virtualization/VirtualizationSystem.h
--- a/Engine/Source/Runtime/Core/Public/Virtualization/VirtualizationSystem.h
+++ b/Engine/Source/Runtime/Core/Public/Virtualization/VirtualizationSystem.h
@@ -142,7 +142,7 @@ public:
 	virtual uint64 GetPayloadSize(const FIoHash& Identifier) = 0;
 };
 
-/** The result opf a push operation on a single payload */
+/** The result of a push operation on a single payload */
 struct FPushResult
 {
 	/** 
@@ -153,7 +153,11 @@ struct FPushResult
 	enum class EStatus : int8
 	{
 		/** The push operation caused an error */
-		Error			= -3,
+		Error			= -5,
+		/** The push operation storage type is disabled */
+		ProcessDisabled	= -4,
+		/** There were no backends that supported the push operation */
+		NoBackend		= -3,
 		/** The push operation was run on an invalid payload id */
 		Invalid			= -2,
 		/** The payload was rejected by the filtering system */
@@ -174,6 +178,16 @@ struct FPushResult
 		return FPushResult(EStatus::Error);
 	}
 
+	static FPushResult GetAsProcessDisabled()
+	{
+		return FPushResult(EStatus::ProcessDisabled);
+	}
+
+	static FPushResult GetAsNoBackend()
+	{
+		return FPushResult(EStatus::NoBackend);
+	}
+
 	static FPushResult GetAsInvalid()
 	{
 		return FPushResult(EStatus::Invalid);
