Add a new overload to IVirtualizationSystem::TryVirtualizePackages, which takes additional options (via a bitfield enum) and returns more info about the resulting process. The original version has been drepcated.

#rnx

- This allows us to provide the user with more ways to customize the virtualization and return more detailed info about it if the calling code wishes to log additional info. In both cases we can extend the options and the data returned without changing the API.
- Previously if we virtualized a package that was not checked out in revision control we would warn the user and then skip updating the package file on disk. This means the payloads would be uploaded but the user would be left with no local changes. Since sometimes we know we don't need to check out any package (virtualizing the packages in a change list for example) we don't want to always incur the cost of polling reivision control to see which packages do need checking out. This is why we now allow the caller to request package files be checked out via the new options enum EVirtualizationOptions.
-- If the EVirtualizationOptions::Checkout flag is provided we will poll the revision control status of all package files and then check out those which need it.
-- We still check if packages can be modified and warn the user if they can't, as package files could be locked in other ways.
- Added a new utility function to SourceControlUtilties to make it easier to check out packages. There is similar functionality elsewhere in the code base but the virtualization module is too low level to make use of it.
- Updated existing code that calls ::TryVirtualizePackages and add cases of ''using namespace UE::Virtualization' where required to improve readability.
- The UnrealVirtualizationTool now supports a new cmdline option "-checkout" that can be used when virtualizing packages. This will checkout any package that was actually virtualized so the result can be saved back out to the workspace domain. This means we no longer require the caller to have checked out the packages before running the tool.

diff --git a/Engine/Source/Developer/Virtualization/Private/PackageVirtualizationProcess.cpp b/Engine/Source/Developer/Virtualization/Private/PackageVirtualizationProcess.cpp
--- a/Engine/Source/Developer/Virtualization/Private/PackageVirtualizationProcess.cpp
+++ b/Engine/Source/Developer/Virtualization/Private/PackageVirtualizationProcess.cpp
@@ -19,6 +19,7 @@
 #include "UObject/UObjectGlobals.h"
 #include "Virtualization/VirtualizationSystem.h"
 #include "VirtualizationManager.h"
+#include "VirtualizationSourceControlUtilities.h"
 
 #define LOCTEXT_NAMESPACE "Virtualization"
 
@@ -154,7 +155,7 @@ private:
 	TMap<FIoHash, FPayloadData> PayloadLookupTable;
 };
 
-void VirtualizePackages(TConstArrayView<FString> PackagePaths, TArray<FText>& OutErrors)
+void VirtualizePackages(TConstArrayView<FString> PackagePaths, EVirtualizationOptions Options, FVirtualizationResult& OutResultInfo)
 {
 	TRACE_CPUPROFILER_EVENT_SCOPE(UE::Virtualization::VirtualizePackages);
 
@@ -168,7 +169,7 @@ void VirtualizePackages(TConstArrayView<FString> PackagePaths, TArray<FText>& Ou
 	Progress.MakeDialog();
 
 	// Other systems may have added errors to this array, we need to check so later we can determine if this function added any additional errors.
-	const int32 NumErrors = OutErrors.Num();
+	const int32 NumErrors = OutResultInfo.GetNumErrors();
 
 	struct FPackageInfo
 	{
@@ -215,7 +216,7 @@ void VirtualizePackages(TConstArrayView<FString> PackagePaths, TArray<FText>& Ou
 				{
 					FText Message = FText::Format(LOCTEXT("Virtualization_PkgHasReferences", "Cannot virtualize the package '{1}' as it has referenced payloads in the trailer"),
 						FText::FromString(PackagePath.GetDebugName()));
-					OutErrors.Add(Message);
+					OutResultInfo.AddError(MoveTemp(Message));
 					return;
 				}
 
@@ -300,7 +301,7 @@ void VirtualizePackages(TConstArrayView<FString> PackagePaths, TArray<FText>& Ou
 		if (!System.PushData(PayloadsToSubmit, EStorageType::Persistent))
 		{
 			FText Message = LOCTEXT("Virtualization_PushFailure", "Failed to push payloads");
-			OutErrors.Add(Message);
+			OutResultInfo.AddError(MoveTemp(Message));
 			return;
 		}
 
@@ -331,7 +332,7 @@ void VirtualizePackages(TConstArrayView<FString> PackagePaths, TArray<FText>& Ou
 					FText Message = FText::Format(	LOCTEXT("Virtualization_UpdateStatusFailed", "Unable to update the status for the payload '{0}' in the package '{1}'"),
 													FText::FromString(LexToString(Request.GetIdentifier())),
 													FText::FromString(PackageInfo.Path.GetDebugName()));
-					OutErrors.Add(Message);
+					OutResultInfo.AddError(MoveTemp(Message));
 					return;
 				}
 			}
@@ -352,7 +353,7 @@ void VirtualizePackages(TConstArrayView<FString> PackagePaths, TArray<FText>& Ou
 
 		const FPackagePath& PackagePath = PackageInfo.Path; // No need to validate path, we checked this earlier
 
-		FString NewPackagePath = DuplicatePackageWithUpdatedTrailer(PackagePath.GetLocalFullPath(), PackageInfo.Trailer, OutErrors);
+		FString NewPackagePath = DuplicatePackageWithUpdatedTrailer(PackagePath.GetLocalFullPath(), PackageInfo.Trailer, OutResultInfo.Errors);
 
 		if (!NewPackagePath.IsEmpty())
 		{
@@ -369,22 +370,44 @@ void VirtualizePackages(TConstArrayView<FString> PackagePaths, TArray<FText>& Ou
 
 	UE_LOG(LogVirtualization, Display, TEXT("%d package(s) had their trailer container modified and need to be updated"), PackagesToReplace.Num());
 
-	if (NumErrors == OutErrors.Num())
+	if (NumErrors == OutResultInfo.GetNumErrors())
 	{
 		// TODO: Consider using the SavePackage model (move the original, then replace, so we can restore all of the original packages if needed)
 		// having said that, once a package is in PackagesToReplace it should still be safe to submit so maybe we don't need this level of protection?
 
+		
 		// We need to reset the loader of any package that we want to re-save over
-		for (int32 Index = 0; Index < PackagesToReplace.Num(); ++Index)
+		for (const TPair<FPackagePath, FString>& Pair : PackagesToReplace)
 		{
-			const TPair<FPackagePath, FString>& Pair = PackagesToReplace[Index];
-
 			UPackage* Package = FindObjectFast<UPackage>(nullptr, Pair.Key.GetPackageFName());
 			if (Package != nullptr)
 			{
 				UE_LOG(LogVirtualization, Verbose, TEXT("Detaching '%s' from disk so that it can be virtualized"), *Pair.Key.GetDebugName());
 				ResetLoadersForSave(Package, *Pair.Key.GetLocalFullPath());
 			}
+		}
+
+		// Should we try to check out packages from revision control?
+		if (EnumHasAnyFlags(Options, EVirtualizationOptions::Checkout))
+		{
+			TArray<FString> FilesToCheckState;
+			FilesToCheckState.Reserve(PackagesToReplace.Num());
+
+			for (const TPair<FPackagePath, FString>& Pair : PackagesToReplace)
+			{
+				FilesToCheckState.Add(Pair.Key.GetLocalFullPath());
+			}
+
+			if (!TryCheckoutFiles(FilesToCheckState, OutResultInfo.Errors, &OutResultInfo.CheckedOutPackages))
+			{
+				return;
+			}
+		}
+
+		// Now check to see if there are package files that cannot be edited because they are read only
+		for (int32 Index = 0; Index < PackagesToReplace.Num(); ++Index)
+		{
+			const TPair<FPackagePath, FString>& Pair = PackagesToReplace[Index];
 
 			if (!CanWriteToFile(Pair.Key.GetLocalFullPath()))
 			{
@@ -408,18 +431,22 @@ void VirtualizePackages(TConstArrayView<FString> PackagePaths, TArray<FText>& Ou
 			const FString OriginalPackagePath = Iterator.Key.GetLocalFullPath();
 			const FString& NewPackagePath = Iterator.Value;
 
-			if (!IFileManager::Get().Move(*OriginalPackagePath, *NewPackagePath))
+			if (IFileManager::Get().Move(*OriginalPackagePath, *NewPackagePath))
+			{
+				OutResultInfo.VirtualizedPackages.Add(OriginalPackagePath);
+			}
+			else
 			{
 				FText Message = FText::Format(	LOCTEXT("Virtualization_MoveFailed", "Unable to replace the package '{0}' with the virtualized version"),
 												FText::FromString(Iterator.Key.GetDebugName()));
-				OutErrors.Add(Message);
+				OutResultInfo.AddError(MoveTemp(Message));
 				continue;
 			}
 		}
 	}
 
-	const double TimeInSeconds = FPlatformTime::Seconds() - StartTime;
-	UE_LOG(LogVirtualization, Verbose, TEXT("Virtualization pre submit check took %.3f(s)"), TimeInSeconds);
+	OutResultInfo.TimeTaken = FPlatformTime::Seconds() - StartTime;
+	UE_LOG(LogVirtualization, Verbose, TEXT("Virtualization pre submit check took %.3f(s)"), OutResultInfo.TimeTaken);
 }
 
 } // namespace UE::Virtualization
diff --git a/Engine/Source/Developer/Virtualization/Private/PackageVirtualizationProcess.h b/Engine/Source/Developer/Virtualization/Private/PackageVirtualizationProcess.h
--- a/Engine/Source/Developer/Virtualization/Private/PackageVirtualizationProcess.h
+++ b/Engine/Source/Developer/Virtualization/Private/PackageVirtualizationProcess.h
@@ -9,6 +9,9 @@
 namespace UE::Virtualization
 {
 
-void VirtualizePackages(TConstArrayView<FString> PackagePaths, TArray<FText>& OutErrors);
+enum class EVirtualizationOptions : uint32;
+struct FVirtualizationResult;
+
+void VirtualizePackages(TConstArrayView<FString> PackagePaths, EVirtualizationOptions Options, FVirtualizationResult& OutResultInfo);
 
 } // namespace UE::Virtualization
diff --git a/Engine/Source/Developer/Virtualization/Private/VirtualizationManager.cpp b/Engine/Source/Developer/Virtualization/Private/VirtualizationManager.cpp
--- a/Engine/Source/Developer/Virtualization/Private/VirtualizationManager.cpp
+++ b/Engine/Source/Developer/Virtualization/Private/VirtualizationManager.cpp
@@ -848,23 +848,20 @@ EQueryResult FVirtualizationManager::QueryPayloadStatuses(TArrayView<const FIoHa
 	return EQueryResult::Success;
 }
 
-EVirtualizationResult FVirtualizationManager::TryVirtualizePackages(TConstArrayView<FString> PackagePaths, TArray<FText>& OutDescriptionTags, TArray<FText>& OutErrors)
+EVirtualizationResult FVirtualizationManager::TryVirtualizePackages(TConstArrayView<FString> PackagePaths, EVirtualizationOptions Options, FVirtualizationResult& OutResultInfo)
 {
-	OutDescriptionTags.Reset();
-	OutErrors.Reset();
-
 	if (IsEnabled() && IsPushingEnabled(EStorageType::Persistent))
 	{
-		UE::Virtualization::VirtualizePackages(PackagePaths, OutErrors);
+		UE::Virtualization::VirtualizePackages(PackagePaths, Options, OutResultInfo);
 
-		if (OutErrors.IsEmpty() && !VirtualizationProcessTag.IsEmpty())
+		if (OutResultInfo.Errors.IsEmpty() && !VirtualizationProcessTag.IsEmpty())
 		{
 			FText Tag = FText::FromString(VirtualizationProcessTag);
-			OutDescriptionTags.Add(Tag);
+			OutResultInfo.DescriptionTags.Add(Tag);
 		}
 	}
 
-	return OutErrors.IsEmpty() ? EVirtualizationResult::Success : EVirtualizationResult::Failed;
+	return OutResultInfo.Errors.IsEmpty() ? EVirtualizationResult::Success : EVirtualizationResult::Failed;
 }
 
 ERehydrationResult FVirtualizationManager::TryRehydratePackages(TConstArrayView<FString> PackagePaths, TArray<FText>& OutErrors)
diff --git a/Engine/Source/Developer/Virtualization/Private/VirtualizationManager.h b/Engine/Source/Developer/Virtualization/Private/VirtualizationManager.h
--- a/Engine/Source/Developer/Virtualization/Private/VirtualizationManager.h
+++ b/Engine/Source/Developer/Virtualization/Private/VirtualizationManager.h
@@ -127,7 +127,7 @@ private:
 
 	virtual EQueryResult QueryPayloadStatuses(TArrayView<const FIoHash> Ids, EStorageType StorageType, TArray<EPayloadStatus>& OutStatuses) override;
 
-	virtual EVirtualizationResult TryVirtualizePackages(TConstArrayView<FString> PackagePaths, TArray<FText>& OutDescriptionTags, TArray<FText>& OutErrors) override;
+	virtual EVirtualizationResult TryVirtualizePackages(TConstArrayView<FString> PackagePaths, EVirtualizationOptions Options, FVirtualizationResult& OutResultInfo) override;
 
 	virtual ERehydrationResult TryRehydratePackages(TConstArrayView<FString> PackagePaths, TArray<FText>& OutErrors) override;
 	virtual ERehydrationResult TryRehydratePackages(TConstArrayView<FString> PackagePaths, uint64 PaddingAlignment, TArray<FText>& OutErrors, TArray<FSharedBuffer>& OutPackages, TArray<FRehydrationInfo>* OutInfo) override;
diff --git a/Engine/Source/Developer/Virtualization/Private/VirtualizationSourceControlUtilities.cpp b/Engine/Source/Developer/Virtualization/Private/VirtualizationSourceControlUtilities.cpp
--- a/Engine/Source/Developer/Virtualization/Private/VirtualizationSourceControlUtilities.cpp
+++ b/Engine/Source/Developer/Virtualization/Private/VirtualizationSourceControlUtilities.cpp
@@ -10,6 +10,8 @@
 #include "Misc/PackageSegment.h"
 #include "SourceControlOperations.h"
 
+#define LOCTEXT_NAMESPACE "Virtualization"
+
 // When enabled we push the source control work to the main thread via the task graph system.
 // Note that when enabled this can potentially cause thread locks so is no a production ready 
 // fix. Realistically we need to fix the SourceControl API to accept requests from any thread 
@@ -147,3 +149,67 @@ bool FVirtualizationSourceControlUtilities::SyncPayloadSidecarFileInternal(const
 }
 
 } // namespace UE::Virtualization::Experimental
+
+namespace UE::Virtualization
+{
+
+bool TryCheckoutFiles(const TArray<FString>& FilesToCheckState, TArray<FText>& OutErrors, TArray<FString>* OutFilesCheckedOut)
+{
+	TArray<FSourceControlStateRef> PathStates;
+	PathStates.Reserve(FilesToCheckState.Num());
+
+	ISourceControlProvider& SCCProvider = ISourceControlModule::Get().GetProvider();
+
+	if (!SCCProvider.IsEnabled())
+	{
+		FText Message = LOCTEXT("VA_RevDisabled", "Cannot check out packages as revision control is disabled");
+		OutErrors.Add(MoveTemp(Message));
+
+		return false;
+	}
+
+	ECommandResult::Type UpdateResult = SCCProvider.GetState(FilesToCheckState, PathStates, EStateCacheUsage::ForceUpdate);
+	if (UpdateResult != ECommandResult::Type::Succeeded)
+	{
+		FText Message = LOCTEXT("VA_FileState", "Failed to find the state of package files in revision control when trying to check them out");
+		OutErrors.Add(MoveTemp(Message));
+
+		return false;
+	}
+
+	TArray<FString> FilesToCheckout;
+	FilesToCheckout.Reserve(PathStates.Num());
+
+	for (const FSourceControlStateRef& State : PathStates)
+	{
+		if (State->IsSourceControlled() && !State->CanEdit() && State->CanCheckout())
+		{
+			FilesToCheckout.Add(State->GetFilename());
+		}
+	}
+
+	if (!FilesToCheckout.IsEmpty())
+	{
+		ECommandResult::Type CheckoutResult = SCCProvider.Execute(ISourceControlOperation::Create<FCheckOut>(), FilesToCheckout);
+		if (CheckoutResult == ECommandResult::Type::Succeeded)
+		{
+			if (OutFilesCheckedOut != nullptr)
+			{
+				*OutFilesCheckedOut = MoveTemp(FilesToCheckout);
+			}	
+		}
+		else
+		{
+			FText Message = LOCTEXT("VA_Checkout", "Failed to checkout packages from revision control");
+			OutErrors.Add(MoveTemp(Message));
+
+			return false;
+		}
+	}
+
+	return true;
+}
+
+} //namespace UE::Virtualization
+
+#undef LOCTEXT_NAMESPACE
diff --git a/Engine/Source/Developer/Virtualization/Private/VirtualizationSourceControlUtilities.h b/Engine/Source/Developer/Virtualization/Private/VirtualizationSourceControlUtilities.h
--- a/Engine/Source/Developer/Virtualization/Private/VirtualizationSourceControlUtilities.h
+++ b/Engine/Source/Developer/Virtualization/Private/VirtualizationSourceControlUtilities.h
@@ -4,6 +4,22 @@
 
 #include "Virtualization/VirtualizationSystem.h"
 
+namespace UE::Virtualization
+{
+
+/**
+* Utility function to make it easier to check out a number of files from revision control. Note that
+* the function does not do anything to files that are already checked out.
+* 
+ * @param FilesToCheckout		A list of files to check out of revision control
+ * @param OutErrors				The function will place any errors encountered here
+ * @param OutFilesCheckedOut	An optional array, which if provided will be filled in with the list
+ *								of files that were actually checked out from revision control.
+ */
+bool TryCheckoutFiles(const TArray<FString>& FilesToCheckout, TArray<FText>& OutErrors, TArray<FString>* OutFilesCheckedOut);
+
+} // UE::Virtualization
+
 namespace UE::Virtualization::Experimental
 {
 
diff --git a/Engine/Source/Editor/SourceControlWindows/Private/SSourceControlSubmit.cpp b/Engine/Source/Editor/SourceControlWindows/Private/SSourceControlSubmit.cpp
--- a/Engine/Source/Editor/SourceControlWindows/Private/SSourceControlSubmit.cpp
+++ b/Engine/Source/Editor/SourceControlWindows/Private/SSourceControlSubmit.cpp
@@ -38,27 +38,33 @@ static FText GSavedChangeListDescription;
 
 bool TryToVirtualizeFilesToSubmit(const TArray<FString>& FilesToSubmit, FText& Description, FText& OutFailureMsg)
 {
-	// TODO: Once this is removed and not deprecated move the following arrays inside of the System.IsEnabled() scope
-	TArray<FText> PayloadErrors;
-	TArray<FText> DescriptionTags;
+	using namespace UE::Virtualization;
+
+	{
+		TArray<FText> PayloadErrors;
+		TArray<FText> DescriptionTags;
 
 PRAGMA_DISABLE_DEPRECATION_WARNINGS
-	ISourceControlModule::Get().GetOnPreSubmitFinalize().Broadcast(FilesToSubmit, DescriptionTags, PayloadErrors);
+		ISourceControlModule::Get().GetOnPreSubmitFinalize().Broadcast(FilesToSubmit, DescriptionTags, PayloadErrors);
 PRAGMA_ENABLE_DEPRECATION_WARNINGS
+	}
 
-	UE::Virtualization::IVirtualizationSystem& System = UE::Virtualization::IVirtualizationSystem::Get();
+	IVirtualizationSystem& System = IVirtualizationSystem::Get();
 	if (!System.IsEnabled())
 	{
 		return true;
 	}
 
-	UE::Virtualization::EVirtualizationResult Result = System.TryVirtualizePackages(FilesToSubmit, DescriptionTags, PayloadErrors);
-	if (Result == UE::Virtualization::EVirtualizationResult::Success)
+	FVirtualizationResult VirtualizationResults;
+	EVirtualizationOptions VirtualizationOptions = EVirtualizationOptions::None;
+
+	EVirtualizationResult Result = System.TryVirtualizePackages(FilesToSubmit, VirtualizationOptions, VirtualizationResults);
+	if (Result == EVirtualizationResult::Success)
 	{
 		FTextBuilder NewDescription;
 		NewDescription.AppendLine(Description);
 
-		for (const FText& Line : DescriptionTags)
+		for (const FText& Line : VirtualizationResults.DescriptionTags)
 		{
 			NewDescription.AppendLine(Line);
 		}
@@ -69,7 +75,7 @@ PRAGMA_ENABLE_DEPRECATION_WARNINGS
 	}
 	else if (System.AllowSubmitIfVirtualizationFailed())
 	{
-		for (const FText& Error : PayloadErrors)
+		for (const FText& Error : VirtualizationResults.Errors)
 		{
 			FMessageLog("SourceControl").Warning(Error);
 		}
@@ -79,7 +85,7 @@ PRAGMA_ENABLE_DEPRECATION_WARNINGS
 	}
 	else
 	{
-		for (const FText& Error : PayloadErrors)
+		for (const FText& Error : VirtualizationResults.Errors)
 		{
 			FMessageLog("SourceControl").Error(Error);
 		}
diff --git a/Engine/Source/Editor/UnrealEd/Private/Commandlets/ContentCommandlets.cpp b/Engine/Source/Editor/UnrealEd/Private/Commandlets/ContentCommandlets.cpp
--- a/Engine/Source/Editor/UnrealEd/Private/Commandlets/ContentCommandlets.cpp
+++ b/Engine/Source/Editor/UnrealEd/Private/Commandlets/ContentCommandlets.cpp
@@ -1780,22 +1780,24 @@ void UResavePackagesCommandlet::CheckInFiles(const TArray<FString>& InFilesToSub
 
 bool UResavePackagesCommandlet::TryVirtualization(const TArray<FString>& FilesToSubmit, FText& InOutDescription)
 {
-	TArray<FText> PayloadErrors;
-	TArray<FText> DescriptionTags;
+	using namespace UE::Virtualization;
 
-	UE::Virtualization::IVirtualizationSystem& System = UE::Virtualization::IVirtualizationSystem::Get();
+	IVirtualizationSystem& System = IVirtualizationSystem::Get();
 	if (!System.IsEnabled())
 	{
 		return true;
 	}
 
-	UE::Virtualization::EVirtualizationResult Result = System.TryVirtualizePackages(FilesToSubmit, DescriptionTags, PayloadErrors);
-	if (Result == UE::Virtualization::EVirtualizationResult::Success)
+	FVirtualizationResult VirtualizationResults;
+	EVirtualizationOptions VirtualizationOptions = EVirtualizationOptions::None;
+
+	EVirtualizationResult Result = System.TryVirtualizePackages(FilesToSubmit, VirtualizationOptions, VirtualizationResults);
+	if (Result == EVirtualizationResult::Success)
 	{
 		FTextBuilder NewDescription;
 		NewDescription.AppendLine(InOutDescription);
 
-		for (const FText& Line : DescriptionTags)
+		for (const FText& Line : VirtualizationResults.DescriptionTags)
 		{
 			NewDescription.AppendLine(Line);
 		}
@@ -1806,7 +1808,7 @@ bool UResavePackagesCommandlet::TryVirtualization(const TArray<FString>& FilesTo
 	}
 	else if (System.AllowSubmitIfVirtualizationFailed())
 	{
-		for (const FText& Error : PayloadErrors)
+		for (const FText& Error : VirtualizationResults.Errors)
 		{
 			UE_LOG(LogContentCommandlet, Warning, TEXT("%s"), *Error.ToString());
 		}
@@ -1816,7 +1818,7 @@ bool UResavePackagesCommandlet::TryVirtualization(const TArray<FString>& FilesTo
 	}
 	else
 	{
-		for (const FText& Error : PayloadErrors)
+		for (const FText& Error : VirtualizationResults.Errors)
 		{
 			UE_LOG(LogContentCommandlet, Error, TEXT("%s"), *Error.ToString());
 		}
diff --git a/Engine/Source/Programs/UnrealVirtualizationTool/Private/Commands/VirtualizeCommand.cpp b/Engine/Source/Programs/UnrealVirtualizationTool/Private/Commands/VirtualizeCommand.cpp
--- a/Engine/Source/Programs/UnrealVirtualizationTool/Private/Commands/VirtualizeCommand.cpp
+++ b/Engine/Source/Programs/UnrealVirtualizationTool/Private/Commands/VirtualizeCommand.cpp
@@ -79,6 +79,10 @@ bool FVirtualizeCommand::Initialize(const TCHAR* CmdLine)
 		{
 			bShouldSubmitChangelist = true;
 		}
+		else if (Switch == TEXT("Checkout"))
+		{
+			bShouldCheckout = true;
+		}
 	}
 
 	// Process the provided changelist if one was found
@@ -156,13 +160,19 @@ bool FVirtualizeCommand::Run(const TArray<FProject>& Projects)
 
 		TArray<FString> ProjectPackages = Project.GetAllPackages();
 
-		TArray<FText> Errors;
-		UE::Virtualization::IVirtualizationSystem::Get().TryVirtualizePackages(ProjectPackages, DescriptionTags, Errors);
+		EVirtualizationOptions Options = EVirtualizationOptions::None;
+		if(bShouldCheckout)
+		{
+			Options |= EVirtualizationOptions::Checkout;
+		}
+
+		FVirtualizationResult ResultInfo;
+		UE::Virtualization::IVirtualizationSystem::Get().TryVirtualizePackages(ProjectPackages, Options, ResultInfo);
 
-		if (!Errors.IsEmpty())
+		if (!ResultInfo.Errors.IsEmpty())
 		{
 			UE_LOG(LogVirtualizationTool, Error, TEXT("The virtualization process failed with the following errors:"));
-			for (const FText& Error : Errors)
+			for (const FText& Error : ResultInfo.Errors)
 			{
 				UE_LOG(LogVirtualizationTool, Error, TEXT("\t%s"), *Error.ToString());
 			}
diff --git a/Engine/Source/Programs/UnrealVirtualizationTool/Private/Commands/VirtualizeCommand.h b/Engine/Source/Programs/UnrealVirtualizationTool/Private/Commands/VirtualizeCommand.h
--- a/Engine/Source/Programs/UnrealVirtualizationTool/Private/Commands/VirtualizeCommand.h
+++ b/Engine/Source/Programs/UnrealVirtualizationTool/Private/Commands/VirtualizeCommand.h
@@ -37,6 +37,7 @@ protected:
 	FString SourceChangelistNumber;
 
 	bool bShouldSubmitChangelist = false;
+	bool bShouldCheckout = false;
 };
 
 /** Used to convert the legacy command line -Mode=Changelist to the new style of command */
diff --git a/Engine/Source/Runtime/Core/Private/Virtualization/VirtualizationSystem.cpp b/Engine/Source/Runtime/Core/Private/Virtualization/VirtualizationSystem.cpp
--- a/Engine/Source/Runtime/Core/Private/Virtualization/VirtualizationSystem.cpp
+++ b/Engine/Source/Runtime/Core/Private/Virtualization/VirtualizationSystem.cpp
@@ -111,12 +111,10 @@ public:
 		return EQueryResult::Failure_NotImplemented;
 	}
 
-	virtual EVirtualizationResult TryVirtualizePackages(TConstArrayView<FString> PackagePaths, TArray<FText>& OutDescriptionTags, TArray<FText>& OutErrors) override
+	virtual EVirtualizationResult TryVirtualizePackages(TConstArrayView<FString> PackagePaths, EVirtualizationOptions Options, FVirtualizationResult& OutResultInfo) override
 	{
-		OutDescriptionTags.Reset();
-		OutErrors.Reset();
-
-		OutErrors.Add(FText::FromString(TEXT("Calling ::TryVirtualizePackages on FNullVirtualizationSystem")));
+		OutResultInfo.Reset();
+		OutResultInfo.Errors.Add(FText::FromString(TEXT("Calling ::TryVirtualizePackages on FNullVirtualizationSystem")));
 
 		return EVirtualizationResult::Failed;
 	}
@@ -124,7 +122,6 @@ public:
 	virtual ERehydrationResult TryRehydratePackages(TConstArrayView<FString> PackagePaths, TArray<FText>& OutErrors) override
 	{
 		OutErrors.Reset();
-
 		OutErrors.Add(FText::FromString(TEXT("Calling ::TryRehydratePackages on FNullVirtualizationSystem")));
 
 		return ERehydrationResult::Failed;
@@ -133,7 +130,6 @@ public:
 	virtual ERehydrationResult TryRehydratePackages(TConstArrayView<FString> PackagePaths, uint64 PaddingAlignment, TArray<FText>& OutErrors, TArray<FSharedBuffer>& OutPackages, TArray<FRehydrationInfo>* OutInfo) override
 	{
 		OutErrors.Reset();
-
 		OutErrors.Add(FText::FromString(TEXT("Calling ::TryRehydratePackages on FNullVirtualizationSystem")));
 
 		return ERehydrationResult::Failed;
diff --git a/Engine/Source/Runtime/Core/Public/Virtualization/VirtualizationSystem.h b/Engine/Source/Runtime/Core/Public/Virtualization/VirtualizationSystem.h
--- a/Engine/Source/Runtime/Core/Public/Virtualization/VirtualizationSystem.h
+++ b/Engine/Source/Runtime/Core/Public/Virtualization/VirtualizationSystem.h
@@ -443,6 +443,62 @@ private:
 	EStatus Status = EStatus::Pending;
 };
 
+/** Contains the results of the package virtualization process */
+struct FVirtualizationResult
+{
+	/** A list of errors encountered during the process */
+	TArray<FText> Errors;
+
+	/** A list of tags associated with the group of packages that were virtualized */
+	TArray<FText> DescriptionTags;
+
+	/** A list of packages that were actually virtualized and had locally stored payloads removed */
+	TArray<FString> VirtualizedPackages;
+	/** A list of packages that were checked out of revision control during the process */
+	TArray<FString> CheckedOutPackages;
+
+	/** The length of time that the process took in seconds */
+	double TimeTaken = 0.0;
+
+	/** Returns members to the default state */
+	void Reset()
+	{
+		Errors.Reset();
+		DescriptionTags.Reset();
+		VirtualizedPackages.Reset();
+		CheckedOutPackages.Reset();
+	}
+
+	/** Add an error message to the result */
+	void AddError(const FText& ErrorMsg)
+	{
+		Errors.Add(ErrorMsg);
+	}
+
+	/** Add an error message to the result */
+	void AddError(FText&& ErrorMsg)
+	{
+		Errors.Add(MoveTemp(ErrorMsg));
+	}
+
+	/** Returns how many errors the process has currently encountered */
+	int32 GetNumErrors() const
+	{
+		return Errors.Num();
+	}
+
+};
+
+/** Options used when virtualizing packages */
+enum class EVirtualizationOptions : uint32
+{
+	None = 0,
+	/** Attempt to check out files from revision control if needed */
+	Checkout = 1 << 0
+};
+
+ENUM_CLASS_FLAGS(EVirtualizationOptions);
+
 /** 
  * The set of parameters to be used when initializing the virtualization system. The 
  * members must remain valid for the duration of the call to ::Initialize. It is not
@@ -691,25 +747,32 @@ public:
 		return QueryPayloadStatuses(Ids, StorageType, OutStatuses) != EQueryResult::Success;
 	}
 
+	UE_DEPRECATED(5.2, "Use The other overload that takes a EVirtualizationOptions parameter")
+	virtual EVirtualizationResult TryVirtualizePackages(TConstArrayView<FString> PackagePaths, TArray<FText>& OutDescriptionTags, TArray<FText>& OutErrors)
+	{
+		FVirtualizationResult Info;
+		EVirtualizationResult Result = TryVirtualizePackages(PackagePaths, EVirtualizationOptions::None, Info);
+
+		OutDescriptionTags = MoveTemp(Info.DescriptionTags);
+		OutErrors = MoveTemp(Info.Errors);
+
+		return Result;
+	}
+	
 	/**
 	 * Runs the virtualization process on a set of packages. All of the packages will be parsed and any found to be containing locally stored
 	 * payloads will have them removed but before they are removed they will be pushed to persistent storage.
-	 * 
-	 * @param PackagePaths			An array of file paths to packages that should be virtualized. If a path resolves to a file that is not 
+	 *
+	 * @param PackagePaths			An array of file paths to packages that should be virtualized. If a path resolves to a file that is not
 	 *								a valid package then it will be silently skipped and will not be considered an error.
-	 * @param OutDescriptionTags	The process may produce description tags associated with the packages which will be placed in this array.
-	 *								These tags can be used to improve logging, or be appended to change list descriptions etc. Note that the 
-	 *								array will be emptied before the process.
-	 *								is run and will not contain any pre-existing entries.
-	 * @param OutErrors				Any error encountered during the process will be added here. If any error is added to the array then it
-	 *								can be assumed that the process will return false. Note that the array will be emptied before the process
-	 *								is run and will not contain any pre-existing entries.
-	 * 
+	 * @param Options				An enum bitfield containing various options for the process. @see EVirtualizationOptions
+	 * @param OutResultInfo			A struct that will contain info about the process. @see FVirtualizationResult
+	 *
 	 * @return						A EVirtualizationResult enum with the status of the process. If the status is not EVirtualizationResult::Success
 	 *								then the parameter OutErrors should contain at least one entry.
 	 */
-	virtual EVirtualizationResult TryVirtualizePackages(TConstArrayView<FString> PackagePaths, TArray<FText>& OutDescriptionTags, TArray<FText>& OutErrors) = 0;
-	
+	virtual EVirtualizationResult TryVirtualizePackages(TConstArrayView<FString> PackagePaths, EVirtualizationOptions Options, FVirtualizationResult& OutResultInfo) = 0;
+
 	/**
 	 * Runs the rehydration process on a set of packages. This involves downloading virtualized payloads and placing them back in the trailer of
 	 * the given packages.
@@ -735,7 +798,7 @@ public:
 	 * @param OutErrors			Any errors encountered while rehydration will be added here
 	 * @param OutPackages		The rehydrated packages as memory buffers. Each entry should match the corresponding entry in PackagePaths. This array is
 	 *							only guaranteed to be correct if the method returns ERehydrationResult::Success.
-	 * @param OutInfo			Information about thge rehydration process, each entry should match the corresponding entry in PackagePaths assuming that
+	 * @param OutInfo			Information about the rehydration process, each entry should match the corresponding entry in PackagePaths assuming that
 	 *							the method returns ERehydrationResult::Success. This parameter is optional, if the information is not required then
 	 *							pass in nullptr to skip.
 	 * @return					ERehydrationResult::Success if the rehydration suceeeds and OutPackages/OutInfocan be trued, otherwise it will return an
