Add a new overload to IVirtualizationSystem::TryRehydratePackages, which takes additional options (via a bitfield enum) and returns more info about the resulting process. The original version has been deprecated.

#rnx

### VA System
- This allows us to provide the user with more ways to customize the rehydration and return more detailed info about it if the calling code wishes to log additional info. In both cases we can extend the options and the data returned without changing the API.
- At the moment the only flag we support is 'Checkout', which requests that the rehydration process checkout any package that it needs to modify rather than giving an error. This means that the user does not need to check packages out from revision control before running the rehydration process.
-- We still check if packages can be modified and warn the user if they can't, as package files could be locked in other ways.
- The rehydration process will now long the time taken if verbose logging is set for the category 'LogVirtualization'

### UnrealVirtualizationTool
- The virtualize command now reports how many packages were checked out if the flag was set.
- The rehydration command now supports a '-Checkout' commandline flag, which when enabled will use the new api to checkout the packages that need to be checked out when rehydrated.

diff --git a/Engine/Source/Developer/Virtualization/Private/PackageRehydrationProcess.cpp b/Engine/Source/Developer/Virtualization/Private/PackageRehydrationProcess.cpp
--- a/Engine/Source/Developer/Virtualization/Private/PackageRehydrationProcess.cpp
+++ b/Engine/Source/Developer/Virtualization/Private/PackageRehydrationProcess.cpp
@@ -14,6 +14,7 @@
 #include "UObject/PackageTrailer.h"
 #include "UObject/UObjectGlobals.h"
 #include "Virtualization/VirtualizationSystem.h"
+#include "VirtualizationSourceControlUtilities.h"
 
 #define LOCTEXT_NAMESPACE "Virtualization"
 
@@ -172,7 +173,7 @@ bool TryRehydrateBuilder(const FString& FilePath, FPackageTrailer& OutTrailer, F
 	return PayloadsHydrated > 0;
 }
 
-void RehydratePackages(TConstArrayView<FString> PackagePaths, TArray<FText>& OutErrors)
+void RehydratePackages(TConstArrayView<FString> PackagePaths, ERehydrationOptions Options, FRehydrationResult& OutResult)
 {
 	TRACE_CPUPROFILER_EVENT_SCOPE(UE::Virtualization::RehydratePackagesOnDisk);
 
@@ -183,51 +184,106 @@ void RehydratePackages(TConstArrayView<FString> PackagePaths, TArray<FText>& Out
 		return;
 	}
 
+	const double StartTime = FPlatformTime::Seconds();
+
 	FScopedSlowTask Progress(1.0f, LOCTEXT("VAHydration_TaskOnDisk", "Re-hydrating Assets On Disk..."));
 	Progress.MakeDialog();
 
+	TArray<TPair<FString, FString>> PackagesToReplace;
+
+	// TODO: Should we consider a way to batch this so many payloads can be downloaded at once?
+	//       Running this over a large project would be much faster if we could batch. An easy way
+	//       to do this might be to gather all of the payloads needed and do a prefetch first?
+
+	// Attempt to rehydrate the packages
 	for (const FString& FilePath : PackagePaths)
 	{
 		FPackageTrailer Trailer;
 		FPackageTrailerBuilder Builder;
 
-		if (TryRehydrateBuilder(FilePath, Trailer, Builder, OutErrors))
+		if (TryRehydrateBuilder(FilePath, Trailer, Builder, OutResult.Errors))
 		{
-			const FString NewPackagePath = DuplicatePackageWithNewTrailer(FilePath, Trailer, Builder, OutErrors);
+			FString NewPackagePath = DuplicatePackageWithNewTrailer(FilePath, Trailer, Builder, OutResult.Errors);
+
+			if (!NewPackagePath.IsEmpty())
+			{
+				PackagesToReplace.Emplace(FilePath, MoveTemp(NewPackagePath));
+			}
+			else
+			{
+				// Error?
+				return;
+			}
+		}
+	}
+
+	// We need to reset the loader of any loaded package that should have its package file replaced
+	for (const TPair<FString, FString>& Pair : PackagesToReplace)
+	{
+		const FString& OriginalFilePath = Pair.Key;
 
-			FString PackageName;
-			if (FPackageName::TryConvertFilenameToLongPackageName(FilePath, PackageName))
+		FString PackageName;
+		if (FPackageName::TryConvertFilenameToLongPackageName(OriginalFilePath, PackageName))
+		{
+			UPackage* Package = FindObjectFast<UPackage>(nullptr, *PackageName);
+			if (Package != nullptr)
 			{
-				UPackage* Package = FindObjectFast<UPackage>(nullptr, *PackageName);
-				if (Package != nullptr)
-				{
-					UE_LOG(LogVirtualization, Verbose, TEXT("Detaching '%s' from disk so that it can be rehydrated"), *FilePath);
-					ResetLoadersForSave(Package, *FilePath);
-				}
+				UE_LOG(LogVirtualization, Verbose, TEXT("Detaching '%s' from disk so that it can be rehydrated"), *OriginalFilePath);
+				ResetLoadersForSave(Package, *OriginalFilePath);
 			}
+		}
+	}
+
+	// Should we try to check out packages from revision control?
+	if (EnumHasAnyFlags(Options, ERehydrationOptions::Checkout))
+	{
+		TArray<FString> FilesToCheckState;
+		FilesToCheckState.Reserve(PackagesToReplace.Num());
+
+		for (const TPair<FString, FString>& Pair : PackagesToReplace)
+		{
+			FilesToCheckState.Add(Pair.Key);
+		}
+
+		if (!TryCheckoutFiles(FilesToCheckState, OutResult.Errors, &OutResult.CheckedOutPackages))
+		{
+			return;
+		}
+	}
+
+	for (const TPair<FString, FString>& Pair : PackagesToReplace)
+	{
+		const FString& OriginalFilePath = Pair.Key;
+		const FString& TempFilePath = Pair.Value;
 
-			if (CanWriteToFile(FilePath))
+		if (CanWriteToFile(OriginalFilePath))
+		{
+			if (IFileManager::Get().Move(*OriginalFilePath, *TempFilePath))
 			{
-				if (!IFileManager::Get().Move(*FilePath, *NewPackagePath))
-				{
-					FText Message = FText::Format(LOCTEXT("VAHydration_MoveFailed", "Unable to replace the package '{0}' with the hydrated version"),
-						FText::FromString(FilePath));
-					OutErrors.Add(Message);
-
-					return;
-				}
+				OutResult.RehydratedPackages.Add(OriginalFilePath);
 			}
 			else
 			{
-				FText Message = FText::Format(
-					LOCTEXT("VAHydration_PackageLocked", "The package file '{0}' has virtualized payloads but is locked for modification and cannot be hydrated"),
-					FText::FromString(FilePath));
+				FText Message = FText::Format(LOCTEXT("VAHydration_MoveFailed", "Unable to replace the package '{0}' with the hydrated version"),
+					FText::FromString(OriginalFilePath));
 
-				OutErrors.Add(Message);
+				OutResult.AddError(MoveTemp(Message));
 				return;
 			}
 		}
+		else
+		{
+			FText Message = FText::Format(
+				LOCTEXT("VAHydration_PackageLocked", "The package file '{0}' has virtualized payloads but is locked for modification and cannot be hydrated"),
+				FText::FromString(OriginalFilePath));
+
+			OutResult.AddError(MoveTemp(Message));
+			return;
+		}
 	}
+
+	OutResult.TimeTaken = FPlatformTime::Seconds() - StartTime;
+	UE_LOG(LogVirtualization, Verbose, TEXT("Rehydration process took %.3f(s)"), OutResult.TimeTaken);
 }
 
 void RehydratePackages(TConstArrayView<FString> PackagePaths, uint64 PaddingAlignment, TArray<FText>& OutErrors, TArray<FSharedBuffer>& OutPackages, TArray<FRehydrationInfo>* OutInfo)
diff --git a/Engine/Source/Developer/Virtualization/Private/PackageRehydrationProcess.h b/Engine/Source/Developer/Virtualization/Private/PackageRehydrationProcess.h
--- a/Engine/Source/Developer/Virtualization/Private/PackageRehydrationProcess.h
+++ b/Engine/Source/Developer/Virtualization/Private/PackageRehydrationProcess.h
@@ -12,8 +12,11 @@ struct FRehydrationInfo;
 namespace UE::Virtualization
 {
 
+enum class ERehydrationOptions : uint32;
+struct FRehydrationResult;
+
 /** This version will rehydrate a package and replace it on disk */
-void RehydratePackages(TConstArrayView<FString> PackagePaths, TArray<FText>& OutErrors);
+void RehydratePackages(TConstArrayView<FString> PackagePaths, ERehydrationOptions Options, FRehydrationResult& OutResult);
 
 /** This version will rehydrate a package but to a memory buffer, the original package on disk will remain untouched */
 void RehydratePackages(TConstArrayView<FString> PackagePaths, uint64 PaddingAlignment, TArray<FText>& OutErrors, TArray<FSharedBuffer>& OutPackages, TArray<FRehydrationInfo>* OutInfo);
diff --git a/Engine/Source/Developer/Virtualization/Private/PackageVirtualizationProcess.cpp b/Engine/Source/Developer/Virtualization/Private/PackageVirtualizationProcess.cpp
--- a/Engine/Source/Developer/Virtualization/Private/PackageVirtualizationProcess.cpp
+++ b/Engine/Source/Developer/Virtualization/Private/PackageVirtualizationProcess.cpp
@@ -446,7 +446,7 @@ void VirtualizePackages(TConstArrayView<FString> PackagePaths, EVirtualizationOp
 	}
 
 	OutResultInfo.TimeTaken = FPlatformTime::Seconds() - StartTime;
-	UE_LOG(LogVirtualization, Verbose, TEXT("Virtualization pre submit check took %.3f(s)"), OutResultInfo.TimeTaken);
+	UE_LOG(LogVirtualization, Verbose, TEXT("Virtualization process took %.3f(s)"), OutResultInfo.TimeTaken);
 }
 
 } // namespace UE::Virtualization
diff --git a/Engine/Source/Developer/Virtualization/Private/VirtualizationManager.cpp b/Engine/Source/Developer/Virtualization/Private/VirtualizationManager.cpp
--- a/Engine/Source/Developer/Virtualization/Private/VirtualizationManager.cpp
+++ b/Engine/Source/Developer/Virtualization/Private/VirtualizationManager.cpp
@@ -839,29 +839,29 @@ EQueryResult FVirtualizationManager::QueryPayloadStatuses(TArrayView<const FIoHa
 	return EQueryResult::Success;
 }
 
-EVirtualizationResult FVirtualizationManager::TryVirtualizePackages(TConstArrayView<FString> PackagePaths, EVirtualizationOptions Options, FVirtualizationResult& OutResultInfo)
+EVirtualizationResult FVirtualizationManager::TryVirtualizePackages(TConstArrayView<FString> PackagePaths, EVirtualizationOptions Options, FVirtualizationResult& OutResult)
 {
 	if (IsEnabled() && IsPushingEnabled(EStorageType::Persistent))
 	{
-		UE::Virtualization::VirtualizePackages(PackagePaths, Options, OutResultInfo);
+		UE::Virtualization::VirtualizePackages(PackagePaths, Options, OutResult);
 
-		if (OutResultInfo.Errors.IsEmpty() && !VirtualizationProcessTag.IsEmpty())
+		if (OutResult.Errors.IsEmpty() && !VirtualizationProcessTag.IsEmpty())
 		{
 			FText Tag = FText::FromString(VirtualizationProcessTag);
-			OutResultInfo.DescriptionTags.Add(Tag);
+			OutResult.DescriptionTags.Add(Tag);
 		}
 	}
 
-	return OutResultInfo.Errors.IsEmpty() ? EVirtualizationResult::Success : EVirtualizationResult::Failed;
+	return OutResult.Errors.IsEmpty() ? EVirtualizationResult::Success : EVirtualizationResult::Failed;
 }
 
-ERehydrationResult FVirtualizationManager::TryRehydratePackages(TConstArrayView<FString> PackagePaths, TArray<FText>& OutErrors)
+ERehydrationResult FVirtualizationManager::TryRehydratePackages(TConstArrayView<FString> PackagePaths, ERehydrationOptions Options, FRehydrationResult& OutResult)
 {
-	OutErrors.Reset();
+	OutResult.Reset();
 
-	UE::Virtualization::RehydratePackages(PackagePaths, OutErrors);
+	UE::Virtualization::RehydratePackages(PackagePaths, Options, OutResult);
 
-	return OutErrors.IsEmpty() ? ERehydrationResult::Success : ERehydrationResult::Failed;
+	return OutResult.Errors.IsEmpty() ? ERehydrationResult::Success : ERehydrationResult::Failed;
 }
 
 ERehydrationResult FVirtualizationManager::TryRehydratePackages(TConstArrayView<FString> PackagePaths, uint64 PaddingAlignment, TArray<FText>& OutErrors, TArray<FSharedBuffer>& OutPackages, TArray<FRehydrationInfo>* OutInfo)
diff --git a/Engine/Source/Developer/Virtualization/Private/VirtualizationManager.h b/Engine/Source/Developer/Virtualization/Private/VirtualizationManager.h
--- a/Engine/Source/Developer/Virtualization/Private/VirtualizationManager.h
+++ b/Engine/Source/Developer/Virtualization/Private/VirtualizationManager.h
@@ -125,9 +125,9 @@ private:
 
 	virtual EQueryResult QueryPayloadStatuses(TArrayView<const FIoHash> Ids, EStorageType StorageType, TArray<EPayloadStatus>& OutStatuses) override;
 
-	virtual EVirtualizationResult TryVirtualizePackages(TConstArrayView<FString> PackagePaths, EVirtualizationOptions Options, FVirtualizationResult& OutResultInfo) override;
+	virtual EVirtualizationResult TryVirtualizePackages(TConstArrayView<FString> PackagePaths, EVirtualizationOptions Options, FVirtualizationResult& OutResult) override;
 
-	virtual ERehydrationResult TryRehydratePackages(TConstArrayView<FString> PackagePaths, TArray<FText>& OutErrors) override;
+	virtual ERehydrationResult TryRehydratePackages(TConstArrayView<FString> PackagePaths, ERehydrationOptions Options, FRehydrationResult& OutResult) override;
 	virtual ERehydrationResult TryRehydratePackages(TConstArrayView<FString> PackagePaths, uint64 PaddingAlignment, TArray<FText>& OutErrors, TArray<FSharedBuffer>& OutPackages, TArray<FRehydrationInfo>* OutInfo) override;
 
 	virtual void DumpStats() const override;
diff --git a/Engine/Source/Editor/VirtualizationEditor/Private/RehydrationMenuEntry.cpp b/Engine/Source/Editor/VirtualizationEditor/Private/RehydrationMenuEntry.cpp
--- a/Engine/Source/Editor/VirtualizationEditor/Private/RehydrationMenuEntry.cpp
+++ b/Engine/Source/Editor/VirtualizationEditor/Private/RehydrationMenuEntry.cpp
@@ -26,7 +26,9 @@ namespace UE::Virtualization
 void RehydratePackages(TArray<FString> SelectedFiles)
 {
 	TArray<FText> Errors;
-	if (IVirtualizationSystem::Get().TryRehydratePackages(SelectedFiles, Errors) == ERehydrationResult::Success)
+
+	FRehydrationResult ResultInfo;
+	if (IVirtualizationSystem::Get().TryRehydratePackages(SelectedFiles, ERehydrationOptions::Checkout, ResultInfo) == ERehydrationResult::Success)
 	{
 		// TODO: At some point ::TryRehydratePackages will return more detail info about the process
 		// when it does we should make a better job at logging it.
diff --git a/Engine/Source/Programs/UnrealVirtualizationTool/Private/Commands/RehydrateCommand.cpp b/Engine/Source/Programs/UnrealVirtualizationTool/Private/Commands/RehydrateCommand.cpp
--- a/Engine/Source/Programs/UnrealVirtualizationTool/Private/Commands/RehydrateCommand.cpp
+++ b/Engine/Source/Programs/UnrealVirtualizationTool/Private/Commands/RehydrateCommand.cpp
@@ -63,6 +63,10 @@ bool FRehydrateCommand::Initialize(const TCHAR* CmdLine)
 		{
 			return TryParseChangelist(ClientSpecName, SwitchValue, Packages, nullptr);
 		}
+		else if (Switch == TEXT("Checkout"))
+		{
+			bShouldCheckout = true;
+		}
 	}
 
 	return true;
@@ -89,7 +93,7 @@ bool FRehydrateCommand::Run(const TArray<FProject>& Projects)
 
 		UE::Virtualization::FInitParams InitParams(ProjectName, EngineConfigWithProject);
 		UE::Virtualization::Initialize(InitParams, UE::Virtualization::EInitializationFlags::ForceInitialize);
-		
+
 		ON_SCOPE_EXIT
 		{
 			UE::Virtualization::Shutdown();
@@ -104,22 +108,34 @@ bool FRehydrateCommand::Run(const TArray<FProject>& Projects)
 
 		UE_LOG(LogVirtualizationTool, Display, TEXT("\t\tAttempting to rehydrate packages..."), ProjectName.ToString());
 
-		TArray<FText> Errors;
-		UE::Virtualization::IVirtualizationSystem::Get().TryRehydratePackages(ProjectPackages, Errors);
 
-		if (!Errors.IsEmpty())
+		ERehydrationOptions Options = ERehydrationOptions::None;
+		if (bShouldCheckout)
+		{
+			Options |= ERehydrationOptions::Checkout;
+		}
+
+		FRehydrationResult Info;
+		UE::Virtualization::IVirtualizationSystem::Get().TryRehydratePackages(ProjectPackages, Options, Info);
+
+		if (!Info.Errors.IsEmpty())
 		{
 			UE_LOG(LogVirtualizationTool, Error, TEXT("The rehydration process failed with the following errors:"));
-			for (const FText& Error : Errors)
+			for (const FText& Error : Info.Errors)
 			{
 				UE_LOG(LogVirtualizationTool, Error, TEXT("\t%s"), *Error.ToString());
 			}
 			return false;
 		}
 
+		if (bShouldCheckout)
+		{
+			UE_LOG(LogVirtualizationTool, Display, TEXT("\t\t%d packages were checked out of revision control"), Info.CheckedOutPackages.Num());
+		}
+
 		Project.UnRegisterMountPoints();
 
-		UE_LOG(LogVirtualizationTool, Display, TEXT("\tRehyration of project packages complete"), ProjectName.ToString());
+		UE_LOG(LogVirtualizationTool, Display, TEXT("\t\tRehyration of project packages complete!"), ProjectName.ToString());
 	}
 
 	return true;
diff --git a/Engine/Source/Programs/UnrealVirtualizationTool/Private/Commands/RehydrateCommand.h b/Engine/Source/Programs/UnrealVirtualizationTool/Private/Commands/RehydrateCommand.h
--- a/Engine/Source/Programs/UnrealVirtualizationTool/Private/Commands/RehydrateCommand.h
+++ b/Engine/Source/Programs/UnrealVirtualizationTool/Private/Commands/RehydrateCommand.h
@@ -26,6 +26,8 @@ private:
 	FString ClientSpecName;
 
 	TArray<FString> Packages;
+
+	bool bShouldCheckout = false;
 };
 
 } // namespace UE::Virtualization
\ No newline at end of file
diff --git a/Engine/Source/Programs/UnrealVirtualizationTool/Private/Commands/VirtualizeCommand.cpp b/Engine/Source/Programs/UnrealVirtualizationTool/Private/Commands/VirtualizeCommand.cpp
--- a/Engine/Source/Programs/UnrealVirtualizationTool/Private/Commands/VirtualizeCommand.cpp
+++ b/Engine/Source/Programs/UnrealVirtualizationTool/Private/Commands/VirtualizeCommand.cpp
@@ -179,7 +179,12 @@ bool FVirtualizeCommand::Run(const TArray<FProject>& Projects)
 			return false;
 		}
 
-		UE_LOG(LogVirtualizationTool, Display, TEXT("\tCheck complete"));
+		if (bShouldCheckout)
+		{
+			UE_LOG(LogVirtualizationTool, Display, TEXT("\t\t%d packages were checked out of revision control"), ResultInfo.CheckedOutPackages.Num());
+		}
+
+		UE_LOG(LogVirtualizationTool, Display, TEXT("\t\tVirtualization of project packages complete!"), ProjectName.ToString());
 	}
 
 	if (bShouldSubmitChangelist)
diff --git a/Engine/Source/Runtime/Core/Private/Virtualization/VirtualizationSystem.cpp b/Engine/Source/Runtime/Core/Private/Virtualization/VirtualizationSystem.cpp
--- a/Engine/Source/Runtime/Core/Private/Virtualization/VirtualizationSystem.cpp
+++ b/Engine/Source/Runtime/Core/Private/Virtualization/VirtualizationSystem.cpp
@@ -119,10 +119,10 @@ public:
 		return EVirtualizationResult::Failed;
 	}
 
-	virtual ERehydrationResult TryRehydratePackages(TConstArrayView<FString> PackagePaths, TArray<FText>& OutErrors) override
+	virtual ERehydrationResult TryRehydratePackages(TConstArrayView<FString> PackagePaths, ERehydrationOptions Options, FRehydrationResult& OutResultInfo) override
 	{
-		OutErrors.Reset();
-		OutErrors.Add(FText::FromString(TEXT("Calling ::TryRehydratePackages on FNullVirtualizationSystem")));
+		OutResultInfo.Reset();
+		OutResultInfo.Errors.Add(FText::FromString(TEXT("Calling ::TryRehydratePackages on FNullVirtualizationSystem")));
 
 		return ERehydrationResult::Failed;
 	}
diff --git a/Engine/Source/Runtime/Core/Public/Virtualization/VirtualizationSystem.h b/Engine/Source/Runtime/Core/Public/Virtualization/VirtualizationSystem.h
--- a/Engine/Source/Runtime/Core/Public/Virtualization/VirtualizationSystem.h
+++ b/Engine/Source/Runtime/Core/Public/Virtualization/VirtualizationSystem.h
@@ -499,6 +499,55 @@ enum class EVirtualizationOptions : uint32
 
 ENUM_CLASS_FLAGS(EVirtualizationOptions);
 
+/** Contains the results of the package virtualization process */
+struct FRehydrationResult
+{
+	/** A list of errors encountered during the process */
+	TArray<FText> Errors;
+
+	/** A list of packages that were actually rehydrated and had payloads added to local storage */
+	TArray<FString> RehydratedPackages;
+	/** A list of packages that were checked out of revision control during the process */
+	TArray<FString> CheckedOutPackages;
+
+	/** The length of time that the process took in seconds */
+	double TimeTaken = 0.0;
+
+	/** Returns members to the default state */
+	void Reset()
+	{
+		Errors.Reset();
+	}
+
+	/** Add an error message to the result */
+	void AddError(const FText& ErrorMsg)
+	{
+		Errors.Add(ErrorMsg);
+	}
+
+	/** Add an error message to the result */
+	void AddError(FText&& ErrorMsg)
+	{
+		Errors.Add(MoveTemp(ErrorMsg));
+	}
+
+	/** Returns how many errors the process has currently encountered */
+	int32 GetNumErrors() const
+	{
+		return Errors.Num();
+	}
+};
+
+/** Options used when rehydrating packages */
+enum class ERehydrationOptions : uint32
+{
+	None = 0,
+	/** Attempt to check out files from revision control if needed */
+	Checkout = 1 << 0
+};
+
+ENUM_CLASS_FLAGS(ERehydrationOptions);
+
 /** 
  * The set of parameters to be used when initializing the virtualization system. The 
  * members must remain valid for the duration of the call to ::Initialize. It is not
@@ -773,20 +822,29 @@ public:
 	 */
 	virtual EVirtualizationResult TryVirtualizePackages(TConstArrayView<FString> PackagePaths, EVirtualizationOptions Options, FVirtualizationResult& OutResultInfo) = 0;
 
+	UE_DEPRECATED(5.2, "Use The other overload that takes a EVirtualizationOptions parameter")
+	virtual ERehydrationResult TryRehydratePackages(TConstArrayView<FString> PackagePaths, TArray<FText>& OutErrors)
+	{
+		FRehydrationResult Info;
+		ERehydrationResult Result = TryRehydratePackages(PackagePaths, ERehydrationOptions::None, Info);
+
+		OutErrors = MoveTemp(Info.Errors);
+
+		return Result;
+	}
+
 	/**
 	 * Runs the rehydration process on a set of packages. This involves downloading virtualized payloads and placing them back in the trailer of
 	 * the given packages.
 	 * 
-	 * @param PackagePaths	An array containing the absolute file paths of packages. It is assumed that the packages have already been checked out
-	 *						of source control (if applicable) and will be writable.
-	 * @param OutErrors		Any error encountered during the process will be added here. If any error is added to the array then it
-	 *						can be assumed that the process will return false. Note that the array will be emptied before the process
-	 *						is run and will not contain any pre-existing entries.
+	 * @param PackagePaths		An array containing the absolute file paths of packages
+	 * @param Options			An enum bitfield containing various options for the process. @see EVirtualizationOptions
+	 * @param OutResultInfo		A struct that will contain info about the process. @see FVirtualizationResult
 	 * 
-	 * @return	A ERehydrationResult enum with the status of the process. If the status indicates any sort of failure then OutErrors should
+	 * @return	A ERehydrationResult enum with the status of the process. If the status indicates any sort of failure then OutResultInfo.Errors should
 	 *			contain at least one entry.
 	 */
-	virtual ERehydrationResult TryRehydratePackages(TConstArrayView<FString> PackagePaths, TArray<FText>& OutErrors) = 0;
+	virtual ERehydrationResult TryRehydratePackages(TConstArrayView<FString> PackagePaths, ERehydrationOptions Options, FRehydrationResult& OutResultInfo) = 0;
 
 	/**
 	 * Rehydrates a number of packages into memory buffers.
