Add documentation for the 'MaxConnections' option found in the FSourceControlBackend.

#rnx

diff --git a/Engine/Source/Developer/Virtualization/Private/VirtualizationSourceControlBackend.h b/Engine/Source/Developer/Virtualization/Private/VirtualizationSourceControlBackend.h
--- a/Engine/Source/Developer/Virtualization/Private/VirtualizationSourceControlBackend.h
+++ b/Engine/Source/Developer/Virtualization/Private/VirtualizationSourceControlBackend.h
@@ -44,6 +44,10 @@ class FSemaphore;
  *								in milliseconds. Remember that the max length of time that the process
  *								can stall attempting to download a payload file is 
  *								RetryCount * RetryWaitTime; [Default=100ms]
+ * MaxConnections [int32]		The maximum number of perforce connections that can be made concurrently. 
+ *								If this value is exceeded then additional requests will wait until a 
+ *								connection becomes free. Pass in '-1' to disable this feature and allow
+ *								unlimited connections. [Default=8]
  * BatchCount [int32]			The max number of payloads that can be pushed to source control in a
  *								single submit. If the number of payloads in a request batch exceeds
  *								this size then it will be split into multiple smaller batches. [Default=100]
