The VA filesystem backend no longer requires that the 'Path' option be setup in the config file.

#rnx

- If no 'Path' option is found the path will default to the projects saved directory, under a 'VAPayloads' subdirectory
- This works as an okay default if the backend is being used for cached storage.

diff --git a/Engine/Source/Developer/Virtualization/Private/VirtualizationFileBackend.cpp b/Engine/Source/Developer/Virtualization/Private/VirtualizationFileBackend.cpp
--- a/Engine/Source/Developer/Virtualization/Private/VirtualizationFileBackend.cpp
+++ b/Engine/Source/Developer/Virtualization/Private/VirtualizationFileBackend.cpp
@@ -20,8 +20,7 @@ bool FFileSystemBackend::Initialize(const FString& ConfigEntry)
 {
 	if (!FParse::Value(*ConfigEntry, TEXT("Path="), RootDirectory))
 	{
-		UE_LOG(LogVirtualization, Error, TEXT("[%s] 'Path=' not found in the config file"), *GetDebugName());
-		return false;
+		RootDirectory = *WriteToString<512>(FPaths::ProjectSavedDir(), TEXT("VAPayloads"));
 	}
 
 	FString EnvOverrideName;
