Do not raise an error if revision control is disabled. If packages did actually need to be checked out then we will raise the appropriate error when trying to modify them.

#rnx

- Fleshed out the comment for the function to better describe the return value.

diff --git a/Engine/Source/Developer/Virtualization/Private/VirtualizationSourceControlUtilities.cpp b/Engine/Source/Developer/Virtualization/Private/VirtualizationSourceControlUtilities.cpp
--- a/Engine/Source/Developer/Virtualization/Private/VirtualizationSourceControlUtilities.cpp
+++ b/Engine/Source/Developer/Virtualization/Private/VirtualizationSourceControlUtilities.cpp
@@ -160,12 +160,10 @@ bool TryCheckoutFiles(const TArray<FString>& FilesToCheckState, TArray<FText>& O
 
 	ISourceControlProvider& SCCProvider = ISourceControlModule::Get().GetProvider();
 
+	// Early out if revision control is disabled
 	if (!SCCProvider.IsEnabled())
 	{
-		FText Message = LOCTEXT("VA_RevDisabled", "Cannot check out packages as revision control is disabled");
-		OutErrors.Add(MoveTemp(Message));
-
-		return false;
+		return true;
 	}
 
 	ECommandResult::Type UpdateResult = SCCProvider.GetState(FilesToCheckState, PathStates, EStateCacheUsage::ForceUpdate);
diff --git a/Engine/Source/Developer/Virtualization/Private/VirtualizationSourceControlUtilities.h b/Engine/Source/Developer/Virtualization/Private/VirtualizationSourceControlUtilities.h
--- a/Engine/Source/Developer/Virtualization/Private/VirtualizationSourceControlUtilities.h
+++ b/Engine/Source/Developer/Virtualization/Private/VirtualizationSourceControlUtilities.h
@@ -15,6 +15,9 @@ namespace UE::Virtualization
  * @param OutErrors				The function will place any errors encountered here
  * @param OutFilesCheckedOut	An optional array, which if provided will be filled in with the list
  *								of files that were actually checked out from revision control.
+ * 
+ * @return	True if the files were checked our OR revision control is not enabled. If revision control was enabled
+  *			and we encountered errors, then the function will return false.
  */
 bool TryCheckoutFiles(const TArray<FString>& FilesToCheckout, TArray<FText>& OutErrors, TArray<FString>* OutFilesCheckedOut);
 
