Allow the revision control VA backend to prioritize requests from the game thread when limited by the number of p4 server connections allowed.

#rnx

- Previously once the number of perforce connections allowed was reached we'd block all additional requests until an outstanding connection as been released at which point a random waiting thread would get the go ahead to make its own connection and so on. If the GameThread was making a request under these conditions the user could be left with a blocked editor while they wait for the GameThread to get a turn.
- This change allows FSourceControlBackend to create a semaphore that gives priority to the GameThread during contention so that the user gets the fastest possible response.

- FSemaphore now takes a bitfield of flags as a parameter in the constructor to allow us to specialize its behavior with the only support option being to prioritize the game thread.
- If the GameThread tries to acquire the semaphore and the prioritization option is enabled we use a slightly different path code path ::AcquireFromGameThread
-- Unlike the normal acquire code path, this path does not increment the count when waiting. This effectively reserves the next open slot for the calling thread.
-- Given that there is only one request on the GameThread at a time we don't have to check Counter once the event has been triggered. The even triggering means we know that one of the slots has been freed and we know we have already reserved that slot so we can just go ahead and return success at that point.
- Added a profile scope to acquire as that would've made the original problem easier to catch.
- Still need to look into replacing this with std::counting_semaphore once we can rely on c++ 20 support.

diff --git a/Engine/Source/Developer/Virtualization/Private/VirtualizationSourceControlBackend.cpp b/Engine/Source/Developer/Virtualization/Private/VirtualizationSourceControlBackend.cpp
--- a/Engine/Source/Developer/Virtualization/Private/VirtualizationSourceControlBackend.cpp
+++ b/Engine/Source/Developer/Virtualization/Private/VirtualizationSourceControlBackend.cpp
@@ -55,13 +55,21 @@ public:
 		EventFailed
 	};
 
+	enum class EFlags : uint32
+	{
+		None					= 0,
+		PrioritizeGameThread	= 1 << 0
+	};
+
+	FRIEND_ENUM_CLASS_FLAGS(EFlags);
+
 	FSemaphore() = delete;
-	explicit FSemaphore(int32 InitialCount)
+	explicit FSemaphore(int32 InitialCount, EFlags Options)
 		: WaitEvent(EEventMode::ManualReset)
 		, Counter(InitialCount)
 		, DebugCount(0)
 	{
-
+		bPrioritizeGameThread = EnumHasAnyFlags(Options, EFlags::PrioritizeGameThread);
 	}
 
 	~FSemaphore()
@@ -72,9 +80,16 @@ public:
 	/** Will block until the calling thread can pass through the semaphore. Note that it might return an error if the WaitEvent fails */
 	EAcquireResult Acquire()
 	{
+		TRACE_CPUPROFILER_EVENT_SCOPE(FSemaphore::Acquire);
+
 		CriticalSection.Lock();
 		DebugCount++;
 
+		if (bPrioritizeGameThread && IsInGameThread())
+		{
+			return AcquireFromGameThread();
+		}
+
 		while (Counter-- <= 0)
 		{
 			Counter++;
@@ -106,13 +121,40 @@ public:
 	}
 
 private:
+
+	inline EAcquireResult AcquireFromGameThread()
+	{
+		if (Counter-- > 0)
+		{
+			CriticalSection.Unlock();
+		}
+		else
+		{
+			WaitEvent->Reset();
+
+			CriticalSection.Unlock();
+
+			if (!WaitEvent->Wait())
+			{
+				--DebugCount;
+				return EAcquireResult::EventFailed;
+			}
+		}
+
+		return EAcquireResult::Success;
+	}
+
 	FEventRef WaitEvent;
 	FCriticalSection CriticalSection;
 
 	std::atomic<int32> Counter;
 	std::atomic<int32> DebugCount;
+
+	bool bPrioritizeGameThread;
 };
 
+ENUM_CLASS_FLAGS(FSemaphore::EFlags);
+
 /** Structure to make it easy to acquire/release a FSemaphore for a given scope */
 struct FSemaphoreScopeLock
 {
@@ -895,7 +937,7 @@ bool FSourceControlBackend::TryApplySettingsFromConfigFiles(const FString& Confi
 
 		if (MaxLimit != INDEX_NONE)
 		{
-			ConcurrentConnectionLimit = MakeUnique<FSemaphore>(MaxLimit);
+			ConcurrentConnectionLimit = MakeUnique<FSemaphore>(MaxLimit, FSemaphore::EFlags::PrioritizeGameThread);
 			UE_LOG(LogVirtualization, Log, TEXT("[%s] Limted to %d concurrent revision control connections"), *GetDebugName(), MaxLimit);
 		}
 		else
