Sanitize filepath input to the virtualization and rehydration processes.

#rnx

- CL 23859876 introduced optimizations that consider a filepath to be invalid if the path contains consecutive path separators (DirA//DirB for example instead of DirA/DirB) and although no issues were seen when testing we fairly quickly ran into a user facing errors when trying to submit work from the editor.
- The users client spec root had a '\' at the end of it where as most internal devs do not. This in turn caused problems when parsing file info from source control as the code was assuming that the client spec root had no trailing separator and was adding one.
- Although the system providing the bad filepaths will be fixed it makes sense to sanitize the filepath input for these two processes as any system, including 3rd party ones could be providing the paths.

diff --git a/Engine/Source/Developer/Virtualization/Private/VirtualizationManager.cpp b/Engine/Source/Developer/Virtualization/Private/VirtualizationManager.cpp
--- a/Engine/Source/Developer/Virtualization/Private/VirtualizationManager.cpp
+++ b/Engine/Source/Developer/Virtualization/Private/VirtualizationManager.cpp
@@ -322,6 +322,33 @@ TArray<FString> ParseEntries(const FString& Data)
 	return Entries;
 }
 
+/** 
+ * FPackagePath has strict requirements on how file paths are formated in order for it to be
+ * able to accept them. In theory the paths that the virtualization system are given are
+ * already in a good shape but as we do not know where they come from we cannot trust them.
+ * This utility will iterate over a list of untrusted filepaths and call various FPaths methods
+ * on them to make sure that they are all correct.
+ */
+TArray<FString> SanitizeFilePaths(TConstArrayView<FString> FilePaths)
+{
+	TRACE_CPUPROFILER_EVENT_SCOPE(SanitizeFilePaths);
+
+	TArray<FString> SanitizedPaths;
+	SanitizedPaths.Reserve(FilePaths.Num());
+
+	for (const FString& Path : FilePaths)
+	{
+		FString SanitizedPath = Path;
+
+		FPaths::NormalizeFilename(SanitizedPath);
+		FPaths::RemoveDuplicateSlashes(SanitizedPath);
+
+		SanitizedPaths.Emplace(MoveTemp(SanitizedPath));
+	}
+
+	return SanitizedPaths;
+}
+
 /** 
  * Utility to make 'StorageType ==  EStorageType::Cache' checks easier while EStorageType::Local continues to exist.
  * When the deprecated value is removed this can also be removed and code calling it can just check for EStorageType::Cache */
@@ -846,7 +873,8 @@ FVirtualizationResult FVirtualizationManager::TryVirtualizePackages(TConstArrayV
 
 	if (IsEnabled() && IsPushingEnabled(EStorageType::Persistent))
 	{
-		UE::Virtualization::VirtualizePackages(PackagePaths, Options, Result);
+		const TArray<FString> SanitizedPaths = SanitizeFilePaths(PackagePaths);
+		UE::Virtualization::VirtualizePackages(SanitizedPaths, Options, Result);
 
 		if (Result.WasSuccessful() && !VirtualizationProcessTag.IsEmpty())
 		{
@@ -861,7 +889,8 @@ FRehydrationResult FVirtualizationManager::TryRehydratePackages(TConstArrayView<
 {
 	FRehydrationResult Result;
 
-	UE::Virtualization::RehydratePackages(PackagePaths, Options, Result);
+	const TArray<FString> SanitizedPaths = SanitizeFilePaths(PackagePaths);
+	UE::Virtualization::RehydratePackages(SanitizedPaths, Options, Result);
 
 	return Result;
 }
