Fix a potential slate assert if a VA connection issue is encountered on a background thread once the editor has initialized.

#rnx

- We were raising a FMessageLog to print a warning, then using ::OnConnectionError to defer the FMessageLog::Notify call until the next editor tick so an easy fix is to pass the warning to ::OnConnectionError and piggyback onto the defferal if needed.
- The messages used to be errors but were demoted during development, now that FSourceControlBackend has SuppressNotifications option and can supress the notifications on a per project basis we can restore the messages to errors.
- When calling ::OnConnectionError we will send the message to FMessageLog if we are currently on the game thread, then defer the notification until next tick as we were doing so previously. The deferral is needed in case the notify comes during editor start up (which for a connection problem is likely) so that the notification appears at the correct time. If we are not on the game thread we should still print the message to the log file so that debugging is easier, only the FMessageLog part needs to be deferred to the GameThread. In this case we need to tell FMessageLog to not print to the log file to avoid duplicates.

diff --git a/Engine/Source/Developer/Virtualization/Private/VirtualizationSourceControlBackend.cpp b/Engine/Source/Developer/Virtualization/Private/VirtualizationSourceControlBackend.cpp
--- a/Engine/Source/Developer/Virtualization/Private/VirtualizationSourceControlBackend.cpp
+++ b/Engine/Source/Developer/Virtualization/Private/VirtualizationSourceControlBackend.cpp
@@ -346,11 +346,11 @@ IVirtualizationBackend::EConnectionStatus FSourceControlBackend::OnConnect()
 				Errors.AppendLine(Msg);
 			}
 
-			FMessageLog Log("LogVirtualization");
-			Log.Warning(FText::Format(LOCTEXT("FailedSourceControlConnection", "Failed to connect to revision control backend with the following errors:\n{0}\nThe revision control backend had trouble connecting!\nTrying logging in with the 'p4 login' command or by using p4vs/UnrealGameSync."),
-				Errors.ToText()));
+			FText Message = FText::Format(LOCTEXT("FailedSourceControlConnection", "Failed to connect to revision control backend with the following errors:\n{0}\nThe revision control backend had trouble connecting!\nTrying logging in with the 'p4 login' command or by using p4vs/UnrealGameSync."),
+				Errors.ToText());
+
+			OnConnectionError(MoveTemp(Message));
 
-			OnConnectionError();
 			return IVirtualizationBackend::EConnectionStatus::Error;
 		}
 	}
@@ -363,24 +363,22 @@ IVirtualizationBackend::EConnectionStatus FSourceControlBackend::OnConnect()
 	TSharedRef<FDownloadFile, ESPMode::ThreadSafe> DownloadCommand = ISourceControlOperation::Create<FDownloadFile>();
 	if (SCCProvider->Execute(DownloadCommand, PayloadMetaInfoPath, EConcurrency::Synchronous) != ECommandResult::Succeeded)
 	{
-		FMessageLog Log("LogVirtualization");
+		FText Msg = FText::Format(LOCTEXT("FailedMetaInfo", "Failed to find 'payload_metainfo.txt' in the depot '{0}'\nThe revision control backend will be unable to pull payloads, is your revision control config set up correctly?"),
+			FText::FromString(DepotRoot));
 
-		Log.Warning(FText::Format(LOCTEXT("FailedMetaInfo", "Failed to find 'payload_metainfo.txt' in the depot '{0}'\nThe revision control backend will be unable to pull payloads, is your revision control config set up correctly?"),
-			FText::FromString(DepotRoot)));
+		OnConnectionError(MoveTemp(Msg));
 
-		OnConnectionError();
 		return IVirtualizationBackend::EConnectionStatus::Error;
 }
 #else
 	TSharedRef<FDownloadFile, ESPMode::ThreadSafe> DownloadCommand = ISourceControlOperation::Create<FDownloadFile>();
 	if (!SCCProvider->TryToDownloadFileFromBackgroundThread(DownloadCommand, PayloadMetaInfoPath))
 	{
-		FMessageLog Log("LogVirtualization");
+		FText Msg = FText::Format(LOCTEXT("FailedMetaInfo", "Failed to find 'payload_metainfo.txt' in the depot '{0}'\nThe revision control backend will be unable to pull payloads, is your revision control config set up correctly?"),
+			FText::FromString(DepotRoot));
 
-		Log.Warning(FText::Format(LOCTEXT("FailedMetaInfo", "Failed to find 'payload_metainfo.txt' in the depot '{0}'\nThe revision control backend will be unable to pull payloads, is your revision control config set up correctly?"),
-			FText::FromString(DepotRoot)));
+		OnConnectionError(MoveTemp(Msg));
 
-		OnConnectionError();
 		return IVirtualizationBackend::EConnectionStatus::Error;
 	}
 #endif //IS_SOURCE_CONTROL_THREAD_SAFE
@@ -388,12 +386,11 @@ IVirtualizationBackend::EConnectionStatus FSourceControlBackend::OnConnect()
 	FSharedBuffer MetaInfoBuffer = DownloadCommand->GetFileData(PayloadMetaInfoPath);
 	if (MetaInfoBuffer.IsNull())
 	{
-		FMessageLog Log("LogVirtualization");
+		FText Msg = FText::Format(LOCTEXT("FailedMetaInfo", "Failed to find 'payload_metainfo.txt' in the depot '{0}'\nThe revision control backend will be unable to pull payloads, is your revision control config set up correctly?"),
+			FText::FromString(DepotRoot));
 
-		Log.Warning(FText::Format(LOCTEXT("FailedMetaInfo", "Failed to find 'payload_metainfo.txt' in the depot '{0}'\nThe revision control backend will be unable to pull payloads, is your revision control config set up correctly?"),
-			FText::FromString(DepotRoot)));
+		OnConnectionError(MoveTemp(Msg));
 
-		OnConnectionError();
 		return IVirtualizationBackend::EConnectionStatus::Error;
 	}
 
@@ -1036,20 +1033,58 @@ bool FSourceControlBackend::FindSubmissionWorkingDir(const FString& ConfigEntry)
 	}
 }
 
-void FSourceControlBackend::OnConnectionError()
+void FSourceControlBackend::OnConnectionError(FText Message)
 {
-	auto Callback = [](float Delta)->bool
+	// We don't know when or where this error might occur. We can currently only create 
+	// FMessageLog on the game thread or risk slate asserts, and if we raise a 
+	// FMessageLog::Notify too early in the editor loading process it won't be shown 
+	// correctly.
+	// To avoid this we push errors to the next editor tick, which will be on the 
+	// GameThread and at a point where notification will work. In addition if we do need
+	// to defer the error message until next tick (rather than just the notification) then
+	// we will write it to the log at the point it is raised rather than the point it is
+	// defered to in an attempt to make the log more readable.
+	if (!Message.IsEmpty())
 	{
-		FMessageLog Log("LogVirtualization");
-		Log.Notify(LOCTEXT("ConnectionError", "Asset virtualization connect errors were encountered, see the message log for more info"));
+		if (::IsInGameThread())
+		{
+			// If we are on the game thread we can post the error message immediately
+			FMessageLog Log("LogVirtualization");
+			Log.Error(Message);
 
-		// This tick callback is one shot, so return false to prevent it being invoked again
-		return false;
-	};
+			// Clear the message so we don't log it a second time
+			Message = FText();
+		}
+		else
+		{
+			// We can only send a FMessageLog on the GameThread so for now just log the error
+			// and we can send it to the FMessageLog system next tick
+			UE_LOG(LogVirtualization, Error, TEXT("%s"), *Message.ToString());
+		}
+	}
 
-	if (!bSuppressNotifications)
+	if (!bSuppressNotifications || !Message.IsEmpty())
 	{
-		FTSTicker::GetCoreTicker().AddTicker(FTickerDelegate::CreateLambda(Callback));
+		auto Callback = [Message, bShouldNotify = bSuppressNotifications](float Delta)->bool
+		{
+			FMessageLog Log("LogVirtualization");
+			Log.SuppressLoggingToOutputLog();
+
+			if (!Message.IsEmpty())
+			{
+				Log.Error(Message);
+			}
+
+			if (!bShouldNotify)
+			{
+				Log.Notify(LOCTEXT("ConnectionError", "Asset virtualization connection errors were encountered, see the message log for more info"));
+			}
+
+			// This tick callback is one shot, so return false to prevent it being invoked again
+			return false;
+		}; 
+	
+		FTSTicker::GetCoreTicker().AddTicker(FTickerDelegate::CreateLambda(MoveTemp(Callback)));
 	}
 }
 
diff --git a/Engine/Source/Developer/Virtualization/Private/VirtualizationSourceControlBackend.h b/Engine/Source/Developer/Virtualization/Private/VirtualizationSourceControlBackend.h
--- a/Engine/Source/Developer/Virtualization/Private/VirtualizationSourceControlBackend.h
+++ b/Engine/Source/Developer/Virtualization/Private/VirtualizationSourceControlBackend.h
@@ -102,7 +102,7 @@ private:
 	bool FindSubmissionWorkingDir(const FString& ConfigEntry);
 
 	/** Will display a FMessage notification to the user on the next valid engine tick to try and keep them aware of connection failures */
-	void OnConnectionError();
+	void OnConnectionError(FText Message);
 
 	/** A source control connection owned by the backend*/
 	TUniquePtr<ISourceControlProvider> SCCProvider;
