Improve the error message when using the console command 'DumpPackagePayloadInfo' if the package does not contain a FPackageTrailer.

- Instead of just giving an error message stating that the trailer could not be found, we now try to return why the trailer cannot be found, if the package is corrupted, or was last saved before the package trailer concept was introduced etc.

diff --git a/Engine/Source/Developer/Virtualization/Private/DumpPackagePayloadInfoCommand.cpp b/Engine/Source/Developer/Virtualization/Private/DumpPackagePayloadInfoCommand.cpp
--- a/Engine/Source/Developer/Virtualization/Private/DumpPackagePayloadInfoCommand.cpp
+++ b/Engine/Source/Developer/Virtualization/Private/DumpPackagePayloadInfoCommand.cpp
@@ -7,6 +7,9 @@
 #include "Misc/PackageName.h"
 #include "HAL/FileManager.h"
 
+#include "UObject/PackageFileSummary.h"
+#include "UObject/PackageResourceManager.h"
+
 namespace UE
 {
 
@@ -30,6 +33,54 @@ FString BytesToString(int64 SizeInBytes)
 	}
 }
 
+void LogPackageError(FArchive* Ar, const TCHAR* DebugName)
+{
+	if (Ar == nullptr)
+	{
+		UE_LOG(LogVirtualization, Error, TEXT("Could not find the package file: '%s'"), *DebugName);
+		return;
+	}
+
+	FPackageFileSummary Summary;
+	*Ar << Summary;
+
+	if (Ar->IsError())
+	{
+		UE_LOG(LogVirtualization, Error, TEXT("Could not find load the package summary from disk: '%s'"), DebugName);
+		return;
+	}
+
+	if (Summary.Tag != PACKAGE_FILE_TAG)
+	{
+		UE_LOG(LogVirtualization, Error, TEXT("Package summary seems to be corrupted: '%s'"), DebugName);
+		return;
+	}
+
+	int32 PackageVersion = Summary.GetFileVersionUE().ToValue();
+	if (PackageVersion >= (int32)EUnrealEngineObjectUE5Version::PAYLOAD_TOC)
+	{
+		UE_LOG(LogVirtualization, Error, TEXT("Package trailer is missing from the package file: '%s'"), DebugName);
+		return;
+	}
+
+	UE_LOG(LogVirtualization, Error, TEXT("Package is tool old (version %d) to have a package trailer (version %d): '%s'"), PackageVersion, EUnrealEngineObjectUE5Version::PAYLOAD_TOC, DebugName);
+}
+
+void LogPackageError(const FPackagePath& Path)
+{
+	const FString DebugName = Path.GetPackageName();
+	TUniquePtr<FArchive> PackageAr = IPackageResourceManager::Get().OpenReadExternalResource(EPackageExternalResource::WorkspaceDomainFile, DebugName);
+	
+	LogPackageError(PackageAr.Get(), *DebugName);
+}
+
+void LogPackageError(const FString& Path)
+{
+	TUniquePtr<FArchive> PackageAr(IFileManager::Get().CreateFileReader(*Path));
+
+	LogPackageError(PackageAr.Get(), *Path);
+}
+
 /**
  * This function is used to write information about package's payloads to the log file. This has no 
  * practical development use and should only be used for debugging purposes. 
@@ -65,7 +116,7 @@ void DumpPackagePayloadInfo(const TArray<FString>& Args)
 		{	
 			if (!FPackageTrailer::TryLoadFromPackage(Path, Trailer))
 			{
-				UE_LOG(LogVirtualization, Error, TEXT("Failed to load the package trailer from: '%s'"), *Path.GetDebugName());
+				LogPackageError(Path);
 				continue;
 			}
 		}
@@ -76,7 +127,7 @@ void DumpPackagePayloadInfo(const TArray<FString>& Args)
 			
 			if (!FPackageTrailer::TryLoadFromFile(PathString, Trailer))
 			{
-				UE_LOG(LogVirtualization, Error, TEXT("Failed to load the package trailer from: '%s'"), *PathString);
+				LogPackageError(PathString);
 				continue;
 			}	
 		}
