Add a warning when virtualizing packages if we find any that predate when package trailers were enabled.

- A common problem that has been reported is that the user has tried to virtualize packages with the virtualization tool but not seen any virtualization occur. When debugged the issue turns out to be that the packages are too old and need to be re-saved.
- If the user submits via the editor then the package will be auto reverted if it has not yet been re-saved  so this problem is not seen in that flow.
- Although this is a problem that will solve itself over time, it would be helpful to the person testing the system if we explicitly log that X number of packages were too old.
- Given how annoying this has proven to people the log message is currently set to warning.
- We don't actually return why a package trailer fails to load from FPackageTrailer::LoadTrailer, as working that out requires additional file reading and most the time we don't really care.
- Instead I added a new utility to the virtualization module so that we can opt into checking the reason.
- At the moment we only report if the package is too old and not for every possible reason.

diff --git a/Engine/Source/Developer/Virtualization/Private/PackageVirtualizationProcess.cpp b/Engine/Source/Developer/Virtualization/Private/PackageVirtualizationProcess.cpp
--- a/Engine/Source/Developer/Virtualization/Private/PackageVirtualizationProcess.cpp
+++ b/Engine/Source/Developer/Virtualization/Private/PackageVirtualizationProcess.cpp
@@ -20,6 +20,7 @@
 #include "Virtualization/VirtualizationSystem.h"
 #include "VirtualizationManager.h"
 #include "VirtualizationSourceControlUtilities.h"
+#include "VirtualizationUtilities.h"
 
 #define LOCTEXT_NAMESPACE "Virtualization"
 
@@ -193,9 +194,11 @@ void VirtualizePackages(TConstArrayView<FString> PackagePaths, EVirtualizationOp
 
 	// From the list of files to submit we need to find all of the valid packages that contain
 	// local payloads that need to be virtualized.
-	int64 TotalPackagesFound = 0;
+	int64 TotalPackagesFound = 0; 
+	int64 TotalOutOfDatePackages = 0;
 	int64 TotalPackageTrailersFound = 0;
 	int64 TotalPayloadsToVirtualize = 0;
+
 	for (const FString& AbsoluteFilePath : PackagePaths)
 	{
 		FPackagePath PackagePath = FPackagePath::FromLocalPath(AbsoluteFilePath);
@@ -233,10 +236,16 @@ void VirtualizePackages(TConstArrayView<FString> PackagePaths, EVirtualizationOp
 					Packages.Emplace(MoveTemp(PkgInfo));
 				}
 			}
+			else if(Utils::FindTrailerFailedReason(PackagePath) == Utils::ETrailerFailedReason::OutOfDate)
+			{
+				
+				TotalOutOfDatePackages++;
+			}
 		}
 	}
 
 	UE_LOG(LogVirtualization, Display, TEXT("Found %" INT64_FMT " package(s), %" INT64_FMT " of which had payload trailers"), TotalPackagesFound, TotalPackageTrailersFound);
+	UE_CLOG(TotalOutOfDatePackages > 0, LogVirtualization, Warning, TEXT("Found %" INT64_FMT " package(s) that are out of date and need resaving"), TotalOutOfDatePackages);
 
 	// TODO: Currently not all of the filtering is done as package save time, so some of the local payloads may not get virtualized.
 	// When/if we move all filtering to package save we can change this log message to state that the local payloads *will* be virtualized.
diff --git a/Engine/Source/Developer/Virtualization/Private/VirtualizationUtilities.cpp b/Engine/Source/Developer/Virtualization/Private/VirtualizationUtilities.cpp
--- a/Engine/Source/Developer/Virtualization/Private/VirtualizationUtilities.cpp
+++ b/Engine/Source/Developer/Virtualization/Private/VirtualizationUtilities.cpp
@@ -4,6 +4,8 @@
 
 #include "IO/IoHash.h"
 #include "Misc/StringBuilder.h"
+#include "UObject/PackageFileSummary.h"
+#include "UObject/PackageResourceManager.h"
 
 namespace UE::Virtualization::Utils
 {
@@ -51,4 +53,29 @@ void GetFormattedSystemError(FStringBuilderBase& SystemErrorMessage)
 	}
 }
 
+ETrailerFailedReason FindTrailerFailedReason(const FPackagePath& PackagePath)
+{
+	TUniquePtr<FArchive> Ar = IPackageResourceManager::Get().OpenReadExternalResource(EPackageExternalResource::WorkspaceDomainFile, PackagePath.GetPackageName());
+
+	if (!Ar)
+	{
+		return ETrailerFailedReason::NotFound;
+	}
+
+	FPackageFileSummary Summary;
+	*Ar << Summary;
+
+	if (Ar->IsError() || Summary.Tag != PACKAGE_FILE_TAG)
+	{
+		return ETrailerFailedReason::InvalidSummary;
+	}
+
+	if (Summary.GetFileVersionUE() < EUnrealEngineObjectUE5Version::PAYLOAD_TOC)
+	{
+		return ETrailerFailedReason::OutOfDate;
+	}
+
+	return ETrailerFailedReason::Unknown;
+}
+
 } // namespace UE::Virtualization::Utils
diff --git a/Engine/Source/Developer/Virtualization/Private/VirtualizationUtilities.h b/Engine/Source/Developer/Virtualization/Private/VirtualizationUtilities.h
--- a/Engine/Source/Developer/Virtualization/Private/VirtualizationUtilities.h
+++ b/Engine/Source/Developer/Virtualization/Private/VirtualizationUtilities.h
@@ -5,6 +5,8 @@
 #include "Containers/StringFwd.h"
 #include "Containers/UnrealString.h"
 
+class FPackagePath;
+
 struct FIoHash;
 
 namespace UE::Virtualization::Utils
@@ -47,4 +49,25 @@ FString PayloadIdToPath(const FIoHash& Id);
  */
 void GetFormattedSystemError(FStringBuilderBase& SystemErrorMessage);
 
+
+enum class ETrailerFailedReason : uint8
+{
+	/** Could not open the package for reading */
+	NotFound,
+	/** The header of the package (summary) could not be read or was otherwise corrupted */
+	InvalidSummary,
+	/** The package predates package version 1002 in which the package trailer was introduced */
+	OutOfDate,
+	/** The reason for the missing package trailer could not be determined */
+	Unknown
+};
+
+/** 
+ * Utility that returns the reason why a given package does not have a package trailer.
+ * Note that it does not actually try to load the trailer or validate that the trailer
+ * is in fact missing and assumes that the caller previously attempted to a call to
+ * load a trailer from the package but it failed.
+ */
+ETrailerFailedReason FindTrailerFailedReason(const FPackagePath& PackagePath);
+
 } // namespace UE::Virtualization::Utils
