Removed the option to opt into virtualization for engine packages. They will be filtered out from virtualization in all cases.

- So far nobody has proposed any worthwhile use case for virtualizing engine content given the following downsides:
-- We cannot ever do so internally as we would be unable to distribute the virtualized packages.
-- Projects and the engine packages cannot always be accessed via the same client spec, leading to problems if we try.
-- Engine content is typically small and not suited to virtualization in the first place.
- If we remove the option then we can exclude the engine packages from our VA commandlets which will give less confusing output (at the moment if you try to virtualize a project containing a single package we will report that several thousand packages were found as we need to parse engine content too)
- Deprecated the ini file options, if we detect them we will warn the user that the option is still present and that engine content now will never virtualize.

diff --git a/Engine/Config/BaseEngine.ini b/Engine/Config/BaseEngine.ini
--- a/Engine/Config/BaseEngine.ini
+++ b/Engine/Config/BaseEngine.ini
@@ -1266,8 +1266,6 @@ BackendGraph=ContentVirtualizationBackendGraph_None
 VirtualizationProcessTag=#virtualized
 AllowSubmitIfVirtualizationFailed=false
 FilterMode=OptOut
-FilterEngineContent=True
-FilterEnginePluginContent=True
 FilterMapContent=True
 LazyInitConnections=False
 ; The following assets should not virtualize until it is confirmed that they do not crash when virtualized payloads are missing
diff --git a/Engine/Source/Developer/Virtualization/Private/VirtualizationManager.cpp b/Engine/Source/Developer/Virtualization/Private/VirtualizationManager.cpp
--- a/Engine/Source/Developer/Virtualization/Private/VirtualizationManager.cpp
+++ b/Engine/Source/Developer/Virtualization/Private/VirtualizationManager.cpp
@@ -471,8 +471,6 @@ FVirtualizationManager::FVirtualizationManager()
 	, BackendGraphName(TEXT("ContentVirtualizationBackendGraph_None"))
 	, VirtualizationProcessTag(TEXT("#virtualized"))
 	, FilteringMode(EPackageFilterMode::OptOut)
-	, bFilterEngineContent(true)
-	, bFilterEnginePluginContent(true)
 	, bFilterMapContent(true)
 	, bAllowSubmitIfVirtualizationFailed(false)
 	, bLazyInitConnections(false)
@@ -1083,28 +1081,20 @@ void FVirtualizationManager::ApplySettingsFromConfigFiles(const FConfigFile& Con
 		UE_LOG(LogVirtualization, Error, TEXT("Failed to load [Core.VirtualizationModule]FilterMode from config file!"));
 	}
 
-	bool bFilterEngineContentFromIni = true;
-	if (ConfigFile.GetBool(LegacyConfigSection, TEXT("FilterEngineContent"), bFilterEngineContentFromIni) ||
-		ConfigFile.GetBool(ConfigSection, TEXT("FilterEngineContent"), bFilterEngineContentFromIni))
+	// Deprecated
 	{
-		bFilterEngineContent = bFilterEngineContentFromIni;
-		UE_LOG(LogVirtualization, Display, TEXT("\tFilterEngineContent : %s"), bFilterEngineContent ? TEXT("true") : TEXT("false"));
-	}
-	else
-	{
-		UE_LOG(LogVirtualization, Error, TEXT("Failed to load [Core.VirtualizationModule].FilterEngineContent from config file!"));
-	}
-	
-	bool bFilterEnginePluginContentFromIni = true;
-	if (ConfigFile.GetBool(LegacyConfigSection, TEXT("FilterEnginePluginContent"), bFilterEnginePluginContentFromIni) ||
-		ConfigFile.GetBool(ConfigSection, TEXT("FilterEnginePluginContent"), bFilterEnginePluginContentFromIni))
-	{
-		bFilterEnginePluginContent = bFilterEnginePluginContentFromIni;
-		UE_LOG(LogVirtualization, Display, TEXT("\tFilterEnginePluginContent : %s"), bFilterEnginePluginContent ? TEXT("true") : TEXT("false"));
-	}
-	else
-	{
-		UE_LOG(LogVirtualization, Error, TEXT("Failed to load [Core.VirtualizationModule].FilterEnginePluginContent from config file!"));
+		bool bDummyValue = true;
+		if (ConfigFile.GetBool(LegacyConfigSection, TEXT("FilterEngineContent"), bDummyValue) ||
+			ConfigFile.GetBool(ConfigSection, TEXT("FilterEngineContent"), bDummyValue))
+		{
+			UE_LOG(LogVirtualization, Warning, TEXT("\tFilterEngineContent is now deprecated (engine content is never virtualized)"));
+		}
+
+		if (ConfigFile.GetBool(LegacyConfigSection, TEXT("FilterEnginePluginContent"), bDummyValue) ||
+			ConfigFile.GetBool(ConfigSection, TEXT("FilterEnginePluginContent"), bDummyValue))
+		{
+			UE_LOG(LogVirtualization, Warning, TEXT("\tFilterEnginePluginContent is now deprecated (engine content is never virtualized)"));
+		}
 	}
 
 	// Optional
@@ -1947,22 +1937,10 @@ bool FVirtualizationManager::ShouldVirtualizePackage(const FPackagePath& Package
 		return true;
 	}
 
-	if (bFilterEngineContent)
-	{
-		// Do not virtualize engine content
-		if (MountPointName.ToView() == TEXT("/Engine/"))
-		{
-			return false;
-		}
-	}
-
-	if (bFilterEnginePluginContent)
+	// Do not virtualize engine content
+	if (MountPointName.ToView() == TEXT("/Engine/") || FPaths::IsUnderDirectory(MountPointPath.ToString(), FPaths::EnginePluginsDir()))
 	{
-		// Do not virtualize engine plugin content
-		if (FPaths::IsUnderDirectory(MountPointPath.ToString(), FPaths::EnginePluginsDir()))
-		{
-			return false;
-		}
+		return false;
 	}
 
 	const UVirtualizationFilterSettings* Settings = GetDefault<UVirtualizationFilterSettings>();
diff --git a/Engine/Source/Developer/Virtualization/Private/VirtualizationManager.h b/Engine/Source/Developer/Virtualization/Private/VirtualizationManager.h
--- a/Engine/Source/Developer/Virtualization/Private/VirtualizationManager.h
+++ b/Engine/Source/Developer/Virtualization/Private/VirtualizationManager.h
@@ -60,8 +60,6 @@ struct FAnalyticsEventAttribute;
  * 
  * [Core.ContentVirtualization]
  * FilterMode=OptIn/OptOut					When 'OptIn' payloads will be virtualized by default, when 'OptOut' they will not be virtualized by default
- * FilterEngineContent=True/False			When true any payload from a package under Engine/Content/.. will be excluded from virtualization
- * FilterEnginePluginContent=True/False		When true any payload from a package under Engine/Plugins/../Content/.. will be excluded from virtualization
  * FilterMapContent=True/False				When true any payload stored in a .umap or _BuildData.uasset file will be excluded from virtualization
  * LazyInitConnections=True/False			When true the backends will not try to make their connections until first used.
  * 
@@ -231,12 +229,6 @@ private:
 	/** The default filtering mode to apply if a payload is not matched with an option in UVirtualizationFilterSettings */
 	EPackageFilterMode FilteringMode;
 
-	/** Should payloads in engine content packages be filtered out and never virtualized */
-	bool bFilterEngineContent;
-	
-	/** Should payloads in engine plugin content packages be filtered out and never virtualized */
-	bool bFilterEnginePluginContent;
-
 	/** Should payloads in .umap files (or associated _BuildData files) be filtered out and never virtualized */
 	bool bFilterMapContent;
 
