Expose more control over when virtualized payloads are placed in cached storage

- The existing config file option 'EnableCacheAfterPull' is now deprecated. If we detect that it is still in the config file we will warn the user about it but continue to apply the value to maintain compatibility.
- Add a new config file option 'EnableCacheOnPull' which controls if a payload is pushed to cached storage after it has been pulled from persistent storage. (Default true)
- Add a new config file option 'EnableCacheOnPush' which controls if a payload is pushed to cached storage before it is pushed to persistent storage. (Default true)
- We now store the info about caching in a single enum bitfield ECachingPolicy.
- Added an up to date comment section about ini file set up for FVirtualizationManager which was previously missing.

diff --git a/Engine/Config/BaseEngine.ini b/Engine/Config/BaseEngine.ini
--- a/Engine/Config/BaseEngine.ini
+++ b/Engine/Config/BaseEngine.ini
@@ -1260,7 +1260,8 @@ LazyInit=False
 
 [Core.VirtualizationModule]
 EnablePayloadVirtualization=True
-EnableCacheAfterPull=True
+EnableCacheOnPull=True
+EnableCacheOnPush=True
 MinPayloadLength=0
 BackendGraph=ContentVirtualizationBackendGraph_None
 VirtualizationProcessTag=#virtualized
diff --git a/Engine/Source/Developer/Virtualization/Private/VirtualizationManager.cpp b/Engine/Source/Developer/Virtualization/Private/VirtualizationManager.cpp
--- a/Engine/Source/Developer/Virtualization/Private/VirtualizationManager.cpp
+++ b/Engine/Source/Developer/Virtualization/Private/VirtualizationManager.cpp
@@ -22,10 +22,13 @@
 
 #define LOCTEXT_NAMESPACE "Virtualization"
 
+
 namespace UE::Virtualization
 {
 UE_REGISTER_VIRTUALIZATION_SYSTEM(UE::Virtualization::FVirtualizationManager, Default);
 
+ENUM_CLASS_FLAGS(FVirtualizationManager::ECachingPolicy);
+
 // Can be defined as 1 by programs target.cs files force the backend connections
 // to lazy initialize on first use rather than when the system is initialized.
 #ifndef UE_VIRTUALIZATION_CONNECTION_LAZY_INIT
@@ -466,7 +469,7 @@ namespace Profiling
 
 FVirtualizationManager::FVirtualizationManager()
 	: bAllowPackageVirtualization(true)
-	, bEnableCacheAfterPull(true)
+	, CachingPolicy(ECachingPolicy::AlwaysCache)
 	, MinPayloadLength(0)
 	, BackendGraphName(TEXT("ContentVirtualizationBackendGraph_None"))
 	, VirtualizationProcessTag(TEXT("#virtualized"))
@@ -555,7 +558,7 @@ PRAGMA_DISABLE_DEPRECATION_WARNINGS
 		case EStorageType::Local:
 PRAGMA_ENABLE_DEPRECATION_WARNINGS
 		case EStorageType::Cache:
-			return !CacheStorageBackends.IsEmpty();
+			return !CacheStorageBackends.IsEmpty() && EnumHasAllFlags(CachingPolicy, ECachingPolicy::CacheOnPush);
 			break;
 
 		case EStorageType::Persistent:
@@ -1016,16 +1019,50 @@ void FVirtualizationManager::ApplySettingsFromConfigFiles(const FConfigFile& Con
 		}
 	}
 
-	bool bEnableCacheAfterPullFromIni = false;
-	if (ConfigFile.GetBool(LegacyConfigSection, TEXT("EnableCacheAfterPull"), bEnableCacheAfterPullFromIni) ||
-		ConfigFile.GetBool(ConfigSection, TEXT("EnableCacheAfterPull"), bEnableCacheAfterPullFromIni))
 	{
-		bEnableCacheAfterPull = bEnableCacheAfterPullFromIni;
-		UE_LOG(LogVirtualization, Display, TEXT("\tCachePulledPayloads : %s"), bEnableCacheAfterPull ? TEXT("true") : TEXT("false"));
+		bool bEnableCacheAfterPullFromIni = false;
+		if (ConfigFile.GetBool(LegacyConfigSection, TEXT("EnableCacheAfterPull"), bEnableCacheAfterPullFromIni) ||
+			ConfigFile.GetBool(ConfigSection, TEXT("EnableCacheAfterPull"), bEnableCacheAfterPullFromIni))
+		{
+			UE_LOG(LogVirtualization, Warning, TEXT("\tEnableCacheAfterPull is now deprecated replace with 'EnableCacheOnPull=True|False'"));
+			if (!bEnableCacheAfterPullFromIni)
+			{
+				EnumRemoveFlags(CachingPolicy, ECachingPolicy::CacheOnPull);
+			}
+		}
+
+		bool bCacheOnPullFromIni = true;
+		if (ConfigFile.GetBool(ConfigSection, TEXT("EnableCacheOnPull"), bCacheOnPullFromIni))
+		{
+			if (!bCacheOnPullFromIni)
+			{
+				EnumRemoveFlags(CachingPolicy, ECachingPolicy::CacheOnPull);
+			}
+
+			UE_LOG(LogVirtualization, Display, TEXT("\tEnableCacheOnPull : %s"),
+				EnumHasAllFlags(CachingPolicy, ECachingPolicy::CacheOnPull) ? TEXT("true") : TEXT("false"));
+		}
+		else
+		{
+			UE_LOG(LogVirtualization, Error, TEXT("Failed to load [Core.VirtualizationModule].EnableCacheOnPull from config file!"));
+		}
 	}
-	else
+
 	{
-		UE_LOG(LogVirtualization, Error, TEXT("Failed to load [Core.VirtualizationModule].EnableCacheAfterPull from config file!"));
+		bool bCacheOnPushFromIni = true;
+		if (ConfigFile.GetBool(ConfigSection, TEXT("EnableCacheOnPush"), bCacheOnPushFromIni))
+		{
+			if (!bCacheOnPushFromIni)
+			{
+				EnumRemoveFlags(CachingPolicy, ECachingPolicy::CacheOnPush);
+			}
+			UE_LOG(LogVirtualization, Display, TEXT("\tEnableCacheOnPush : %s"), 
+				EnumHasAllFlags(CachingPolicy, ECachingPolicy::CacheOnPush) ? TEXT("true") : TEXT("false"));
+		}
+		else
+		{
+			UE_LOG(LogVirtualization, Error, TEXT("Failed to load [Core.VirtualizationModule].EnableCacheOnPush from config file!"));
+		}
 	}
 
 	int64 MinPayloadLengthFromIni = 0;
@@ -1097,7 +1134,6 @@ void FVirtualizationManager::ApplySettingsFromConfigFiles(const FConfigFile& Con
 		}
 	}
 
-	// Optional
 	bool bFilterMapContentFromIni = false;
 	if (ConfigFile.GetBool(ConfigSection, TEXT("FilterMapContent"), bFilterMapContentFromIni))
 	{
@@ -1123,7 +1159,6 @@ void FVirtualizationManager::ApplySettingsFromConfigFiles(const FConfigFile& Con
 		}	
 	}
 
-	// Optional
 	bool bAllowSubmitIfVirtualizationFailedFromIni = true;
 	if (ConfigFile.GetBool(ConfigSection, TEXT("AllowSubmitIfVirtualizationFailed"), bAllowSubmitIfVirtualizationFailedFromIni))
 	{
@@ -1135,7 +1170,6 @@ void FVirtualizationManager::ApplySettingsFromConfigFiles(const FConfigFile& Con
 		UE_LOG(LogVirtualization, Error, TEXT("Failed to load [Core.VirtualizationModule].AllowSubmitIfVirtualizationFailed from config file!"));
 	}
 
-	// Optional
 #if UE_VIRTUALIZATION_CONNECTION_LAZY_INIT == 0
 	bool bLazyInitConnectionsFromIni = true;
 	if (ConfigFile.GetBool(ConfigSection, TEXT("LazyInitConnections"), bLazyInitConnectionsFromIni))
@@ -1860,7 +1894,9 @@ void FVirtualizationManager::PullDataFromAllBackends(TArrayView<FPullRequest> Re
 
 		PullDataFromBackend(*Backend, RequestsCollection.GetRequests());
 
-		TArray<FPushRequest> PayloadsToCache = RequestsCollection.OnPullCompleted(*Backend, bEnableCacheAfterPull);
+		const bool bShouldCache = EnumHasAllFlags(CachingPolicy, ECachingPolicy::CacheOnPull);
+
+		TArray<FPushRequest> PayloadsToCache = RequestsCollection.OnPullCompleted(*Backend, bShouldCache);
 		if (!PayloadsToCache.IsEmpty())
 		{
 			CachePayloads(PayloadsToCache, Backend);
diff --git a/Engine/Source/Developer/Virtualization/Private/VirtualizationManager.h b/Engine/Source/Developer/Virtualization/Private/VirtualizationManager.h
--- a/Engine/Source/Developer/Virtualization/Private/VirtualizationManager.h
+++ b/Engine/Source/Developer/Virtualization/Private/VirtualizationManager.h
@@ -58,10 +58,9 @@ struct FAnalyticsEventAttribute;
  * 
  * Basic Setup:
  * 
- * [Core.ContentVirtualization]
+ * [Core.VirtualizationModule]
  * FilterMode=OptIn/OptOut					When 'OptIn' payloads will be virtualized by default, when 'OptOut' they will not be virtualized by default
  * FilterMapContent=True/False				When true any payload stored in a .umap or _BuildData.uasset file will be excluded from virtualization
- * LazyInitConnections=True/False			When true the backends will not try to make their connections until first used.
  * 
  * PackagePath Setup:
  * 
@@ -79,6 +78,36 @@ struct FAnalyticsEventAttribute;
  * +IncludePackagePaths="/MountPoint/PathTo/ThePackageToInclude"	Includes the specific package '/MountPoint/PathTo/ThePackageToInclude' in the virtualization process
  */
 
+/*
+ * FVirtualizationManager
+ * 
+ * Ini file setup:
+ * 
+ * EnablePayloadVirtualization [bool]:			When true the virtualization process will be enabled (usually when a package is submitted
+												to revision control. [Default=true]
+ * EnableCacheOnPull [bool]:					When true payloads will be pushed to cached storage after being pulled from persistent
+ *												storage. [Default=true]
+ * EnableCacheOnPush [bool]:					When true payloads will be pushed to cached storage right before being pushed to persistent
+ *												storage. [Default=true]
+ * MinPayloadLength  [int64]:					The minimum length (in bytes) that a payload must reach before it can be considered for
+ *												virtualization. Use this to strike a balance between disk space and the number of smaller
+												payloads in your project being virtualized. [Default=0]
+ * BackendGraph [string]:						The name of the backend graph to use. The default graph has no backends and effectively
+												disables the system. It is expected that a project will define the graph that it wants
+												and then set this option [Default=ContentVirtualizationBackendGraph_None]
+ * VirtualizationProcessTag [string]:			The tag to be applied to any set of packages that have had  the virtualization process run
+ *												on them. Typically this means appending the tag to the description of a changelist of 
+ *												packages. This value can be set to an empty string. [Default="#virtualized"]
+ * AllowSubmitIfVirtualizationFailed [bool]:	Revision control submits that trigger the virtualization system can either allow or block
+ *												the submit if the virtualization process fails based on this value. True will allow a
+ *												submit with an error to continue and false will block the submit. Note that by error we mean
+ *												that the packages were not virtualized, not that bad data was produced. [Default=false]
+ * LazyInitConnections [bool]:					When true, backends will not attempt to connect to their services until actually required.
+ *												This can remove lengthy connection steps from the process init phase and then only connect
+ *												if we actually need that service. Note that if this is true then the connection can come from
+ *												any thread, so custom backend code will need to take that into account. [Default=false]
+ */
+
 namespace UE::Virtualization
 {
 class IVirtualizationBackend;
@@ -214,8 +243,22 @@ private:
 	/** Are packages allowed to be virtualized when submitted to source control. Defaults to true. */
 	bool bAllowPackageVirtualization;
 
-	/** Should payloads be cached locally after being pulled from persistent storage? Defaults to true. */
-	bool bEnableCacheAfterPull;
+	enum ECachingPolicy
+	{
+		/** Never push payloads to cached storage */
+		None = 0,
+		/** Cache payloads after they have been pulled from persistent storage */
+		CacheOnPull = 1 << 0,
+		/** Cache payloads right before they are pushed to persistent storage */
+		CacheOnPush = 1 << 1,
+
+		AlwaysCache = CacheOnPull | CacheOnPush
+	};
+
+	FRIEND_ENUM_CLASS_FLAGS(ECachingPolicy);
+
+	/** A bitfield describing when we push payloads to cached storage. */
+	ECachingPolicy CachingPolicy;
 
 	/** The minimum length for a payload to be considered for virtualization. Defaults to 0 bytes. */
 	int64 MinPayloadLength;
