Add progress logging to the reydrate commandlet.

#rnx

- The commandlet runs quite slowly at the moment, logging the progress every 10 seconds or so helps the user see that work is being done and the process hasn't locked up.

diff --git a/Engine/Source/Developer/Virtualization/Private/PackageRehydrationProcess.cpp b/Engine/Source/Developer/Virtualization/Private/PackageRehydrationProcess.cpp
--- a/Engine/Source/Developer/Virtualization/Private/PackageRehydrationProcess.cpp
+++ b/Engine/Source/Developer/Virtualization/Private/PackageRehydrationProcess.cpp
@@ -195,9 +195,13 @@ void RehydratePackages(TConstArrayView<FString> PackagePaths, ERehydrationOption
 	//       Running this over a large project would be much faster if we could batch. An easy way
 	//       to do this might be to gather all of the payloads needed and do a prefetch first?
 
+	double Time = FPlatformTime::Seconds();
+
 	// Attempt to rehydrate the packages
-	for (const FString& FilePath : PackagePaths)
+	for(int32 Index = 0; Index < PackagePaths.Num(); ++Index)
 	{
+		const FString& FilePath = PackagePaths[Index];
+
 		FPackageTrailer Trailer;
 		FPackageTrailerBuilder Builder;
 
@@ -215,6 +219,14 @@ void RehydratePackages(TConstArrayView<FString> PackagePaths, ERehydrationOption
 				return;
 			}
 		}
+
+		if (FPlatformTime::Seconds() - Time > 10.0)
+		{
+			const float ProgressPercent = ((float)Index / (float)PackagePaths.Num()) * 100.0f;
+			UE_LOG(LogVirtualization, Display, TEXT("%d/%d - %.1f%%"), Index, PackagePaths.Num(), ProgressPercent);
+
+			Time = FPlatformTime::Seconds();
+		}
 	}
 
 	// We need to reset the loader of any loaded package that should have its package file replaced
