Add a function to the virtualization utilities that will expand environment variables found in a file path.

#rnx

diff --git a/Engine/Source/Developer/Virtualization/Private/VirtualizationUtilities.cpp b/Engine/Source/Developer/Virtualization/Private/VirtualizationUtilities.cpp
--- a/Engine/Source/Developer/Virtualization/Private/VirtualizationUtilities.cpp
+++ b/Engine/Source/Developer/Virtualization/Private/VirtualizationUtilities.cpp
@@ -2,7 +2,9 @@
 
 #include "VirtualizationUtilities.h"
 
+#include "HAL/PlatformProcess.h"
 #include "IO/IoHash.h"
+#include "Misc/Paths.h"
 #include "Misc/StringBuilder.h"
 #include "UObject/PackageFileSummary.h"
 #include "UObject/PackageResourceManager.h"
@@ -78,4 +80,63 @@ ETrailerFailedReason FindTrailerFailedReason(const FPackagePath& PackagePath)
 	return ETrailerFailedReason::Unknown;
 }
 
+bool ExpandEnvironmentVariables(FStringView InputPath, FStringBuilderBase& OutExpandedPath)
+{
+	while (true)
+	{
+		const int32 EnvVarStart = InputPath.Find(TEXT("$("));
+		if (EnvVarStart == INDEX_NONE)
+		{
+			// If we haven't expanded anything yet we can just copy the input path
+			// If we have expanded then we need to append the remainder of the path
+			if (OutExpandedPath.Len() == 0)
+			{
+				OutExpandedPath = InputPath;
+				return true;
+			}
+			else
+			{
+				OutExpandedPath << InputPath;
+				return true;
+			}
+		}
+
+		const int32 EnvVarEnd = InputPath.Find(TEXT(")"), EnvVarStart + 2);
+		const int32 EnvVarNameLength = EnvVarEnd - (EnvVarStart + 2);
+
+		TStringBuilder<128> EnvVarName;
+		EnvVarName = InputPath.Mid(EnvVarStart + 2, EnvVarNameLength);
+
+		FString EnvVarValue;
+		if (EnvVarName.ToView() == TEXT("Temp") || EnvVarName.ToView() == TEXT("Tmp"))
+		{
+			// On windows the temp envvar is often in 8.3 format
+			// Either we need to expose ::GetLongPathName in some way or we need to consider
+			// calling it in WindowsPlatformMisc::GetEnvironmentVariable.
+			// Until we decide this is a quick work around, check for the Temp envvar and if
+			// it is being requested us the ::UserTempDir function which will convert 8.3 
+			// format correctly.
+			// This should be solved before we consider moving this utility function into core
+			EnvVarValue = FPlatformProcess::UserTempDir();
+
+			FPaths::NormalizeDirectoryName(EnvVarValue);
+		}
+		else
+		{
+			EnvVarValue = FPlatformMisc::GetEnvironmentVariable(EnvVarName.ToString());
+			if (EnvVarValue.IsEmpty())
+			{
+				UE_LOG(LogVirtualization, Warning, TEXT("Could not find environment variable '%s' to expand"), EnvVarName.ToString());
+				OutExpandedPath.Reset();
+				return false;
+			}
+		}
+
+		OutExpandedPath << InputPath.Mid(0, EnvVarStart);
+		OutExpandedPath << EnvVarValue;
+
+		InputPath = InputPath.Mid(EnvVarEnd + 1);
+	}
+}
+
 } // namespace UE::Virtualization::Utils
diff --git a/Engine/Source/Developer/Virtualization/Private/VirtualizationUtilities.h b/Engine/Source/Developer/Virtualization/Private/VirtualizationUtilities.h
--- a/Engine/Source/Developer/Virtualization/Private/VirtualizationUtilities.h
+++ b/Engine/Source/Developer/Virtualization/Private/VirtualizationUtilities.h
@@ -70,4 +70,11 @@ enum class ETrailerFailedReason : uint8
  */
 ETrailerFailedReason FindTrailerFailedReason(const FPackagePath& PackagePath);
 
+/** 
+ * Parse the given path and expand any environment variables found with the format
+ * $(EnvVarName). If the variable cannot be found the function will return false 
+ * and OutExpandedPath will be reset.
+ */
+bool ExpandEnvironmentVariables(FStringView InputPath, FStringBuilderBase& OutExpandedPath);
+
 } // namespace UE::Virtualization::Utils
