Add a new config file option (-WorkingDir=ABC) when setting up a VA source control backend that allows finer grain control over where payloads are submitted from.

#rnx

- This feature was requested so that a team could opt into submitting from their appdata or temp directories instead of the project saved directory (our default)
- Deprecated the option '-SubmitFromTempDir' as it could be replacated with "-WorkingDir=$(Temp)/UnrealEngine/VASubmission"

diff --git a/Engine/Source/Developer/Virtualization/Private/VirtualizationSourceControlBackend.cpp b/Engine/Source/Developer/Virtualization/Private/VirtualizationSourceControlBackend.cpp
--- a/Engine/Source/Developer/Virtualization/Private/VirtualizationSourceControlBackend.cpp
+++ b/Engine/Source/Developer/Virtualization/Private/VirtualizationSourceControlBackend.cpp
@@ -990,18 +990,37 @@ bool FSourceControlBackend::FindSubmissionWorkingDir(const FString& ConfigEntry)
 	// During submission each payload path will be 90 characters in length which will then be appended to
 	// the SubmissionWorkingDir
 
-	SubmissionRootDir = FPlatformMisc::GetEnvironmentVariable(TEXT("UE-VirtualizationWorkingDir"));
-
-	if (!SubmissionRootDir.IsEmpty())
+	FString WorkingDirFromIniFile;
+	if (FParse::Value(*ConfigEntry, TEXT("WorkingDir="), WorkingDirFromIniFile) && !WorkingDirFromIniFile.IsEmpty())
 	{
-		FPaths::NormalizeDirectoryName(SubmissionRootDir);
-		UE_LOG(LogVirtualization, Log, TEXT("[%s] Found Environment Variable: UE-VirtualizationWorkingDir"), *GetDebugName());	
+		TStringBuilder<512> ExpandedPath;
+		if (Utils::ExpandEnvironmentVariables(WorkingDirFromIniFile, ExpandedPath))
+		{
+			SubmissionRootDir = ExpandedPath;
+		}
+		else
+		{
+			UE_LOG(LogVirtualization, Log, TEXT("[%s] Failed to correctly expand 'WorkingDir=%s'"), *GetDebugName(), *WorkingDirFromIniFile);
+		}
 	}
-	else
+	
+	if (SubmissionRootDir.IsEmpty())
+	{
+		SubmissionRootDir = FPlatformMisc::GetEnvironmentVariable(TEXT("UE-VirtualizationWorkingDir"));
+		if (!SubmissionRootDir.IsEmpty())
+		{
+			UE_LOG(LogVirtualization, Log, TEXT("[%s] Found Environment Variable: UE-VirtualizationWorkingDir"), *GetDebugName());
+		}
+	}
+
+	if (SubmissionRootDir.IsEmpty())
 	{
 
 		bool bSubmitFromTempDir = false;
-		FParse::Bool(*ConfigEntry, TEXT("SubmitFromTempDir="), bSubmitFromTempDir);
+		if (FParse::Bool(*ConfigEntry, TEXT("SubmitFromTempDir="), bSubmitFromTempDir))
+		{
+			UE_LOG(LogVirtualization, Warning, TEXT("[%s] Found legacy ini file setting 'SubmitFromTempDir' use '-WorkingDir=$(Temp)/UnrealEngine/VASubmission' instead!"), *GetDebugName());
+		}
 
 		TStringBuilder<260> PathBuilder;
 		if (bSubmitFromTempDir)
@@ -1016,6 +1035,8 @@ bool FSourceControlBackend::FindSubmissionWorkingDir(const FString& ConfigEntry)
 		SubmissionRootDir = PathBuilder;
 	}
 
+	FPaths::NormalizeDirectoryName(SubmissionRootDir);
+
 	if (IFileManager::Get().DirectoryExists(*SubmissionRootDir) || IFileManager::Get().MakeDirectory(*SubmissionRootDir))
 	{
 		UE_LOG(LogVirtualization, Log, TEXT("[%s] Setting '%s' as the working directory"), *GetDebugName(), *SubmissionRootDir);
diff --git a/Engine/Source/Developer/Virtualization/Private/VirtualizationSourceControlBackend.h b/Engine/Source/Developer/Virtualization/Private/VirtualizationSourceControlBackend.h
--- a/Engine/Source/Developer/Virtualization/Private/VirtualizationSourceControlBackend.h
+++ b/Engine/Source/Developer/Virtualization/Private/VirtualizationSourceControlBackend.h
@@ -21,8 +21,8 @@ class FSemaphore;
  * backend, found in Utils::PayloadIdToPath.
  *
  * Ini file setup:
- * 'Name'=(Type=P4SourceControl, DepotRoot="//XXX/", UsePartitionedClient=X, SubmitFromTempDir=X)
- * Where 'Name' is the backend name in the hierarchy and 'XXX' is the path in the source control
+ * 'Name'=(Type=P4SourceControl, DepotRoot="//ABC/")
+ * Where 'Name' is the backend name in the hierarchy and 'ABC' is the path in the source control
  * depot where the payload files are being stored.
  * 
  * Optional Values:
@@ -34,9 +34,10 @@ class FSemaphore;
  *								from will be created as a partitioned workspace which is less overhead
  *								on the source control server. If your server does not support this then
  *								use false. [Default=True]
- * SubmitFromTempDir [bool]:	When set to true, payloads will be submitted from the temp directory of
- *								the current machine and when false the files will be submitted from the
- *								Save directory of the current project. [Default=false]
+ * WorkingDir [string]:			Sets the temp location on disk where payloads are submitted from. The path can
+ *								contain config file expansions (see ConfigCacheIni.cpp ::MatchExpansions) and
+ *								environment variables with the format $(EnvVarName).
+ *								[Default="%GAMESAVEDDIR%/VASubmission"]
  * RetryCount [int32]:			How many times we should try to download a payload before giving up with
  *								an error. Useful when the connection is unreliable but does not experience 
  *								frequent persistent outages. [Default=2]
