Add a config file option '-IgnoreFile=' that allows a project to define the name of the p4ignore file used in the perforce environment.

- We really should just read this value out of the users p4 environment settings, but that would require a source control api change and would need to be done in a source control agnostic way. Letting a project just define the name is a quick way to expose the functionality.
- At some point if we do extend the source control api then this option can be deprecated.

diff --git a/Engine/Source/Developer/Virtualization/Private/VirtualizationSourceControlBackend.cpp b/Engine/Source/Developer/Virtualization/Private/VirtualizationSourceControlBackend.cpp
--- a/Engine/Source/Developer/Virtualization/Private/VirtualizationSourceControlBackend.cpp
+++ b/Engine/Source/Developer/Virtualization/Private/VirtualizationSourceControlBackend.cpp
@@ -180,15 +180,13 @@ private:
 };
 
 /** Utility function to create a directory to submit payloads from. */
-[[nodiscard]] static bool TryCreateSubmissionSessionDirectory(FStringView SessionDirectoryPath)
+[[nodiscard]] static bool TryCreateSubmissionSessionDirectory(FStringView SessionDirectoryPath, FStringView IgnoreFileName)
 {
 	// Write out an ignore file to the submission directory (will create the directory if needed)
 	{
 		TStringBuilder<260> IgnoreFilePath;
 
-		// TODO: We should find if P4IGNORE is actually set and if so extract the filename to use.
-		// This will require extending the source control module
-		FPathViews::Append(IgnoreFilePath, SessionDirectoryPath, TEXT(".p4ignore.txt"));
+		FPathViews::Append(IgnoreFilePath, SessionDirectoryPath, IgnoreFileName);
 
 		// A very basic .p4ignore file that should make sure that we are only submitting valid .upayload files.
 		// 
@@ -597,7 +595,7 @@ bool FSourceControlBackend::PushData(TArrayView<FPushRequest> Requests)
 	TStringBuilder<260> SessionDirectory;
 	FPathViews::Append(SessionDirectory, SubmissionRootDir, SessionGuid);
 
-	if (!TryCreateSubmissionSessionDirectory(SessionDirectory))
+	if (!TryCreateSubmissionSessionDirectory(SessionDirectory, IgnoreFileName))
 	{
 		UE_LOG(LogVirtualization, Error, TEXT("[%s] Failed to created directory '%s' to submit payloads from"), *GetDebugName(), SessionDirectory.ToString());
 		return false;
@@ -968,6 +966,14 @@ bool FSourceControlBackend::TryApplySettingsFromConfigFiles(const FString& Confi
 		}
 	}
 
+	{
+		// TODO: We should just extract this from the perforce environment but that requires extending
+		// the source control api.
+		// Letting the backend define the ignore filename to use is a quicker work around
+		FParse::Value(*ConfigEntry, TEXT("IgnoreFile="), IgnoreFileName);
+		UE_LOG(LogVirtualization, Log, TEXT("[%s] Using '%s' as the p4 ignore file name"), *GetDebugName(), *IgnoreFileName);
+	}
+
 	if (!FindSubmissionWorkingDir(ConfigEntry))
 	{
 		return false;
diff --git a/Engine/Source/Developer/Virtualization/Private/VirtualizationSourceControlBackend.h b/Engine/Source/Developer/Virtualization/Private/VirtualizationSourceControlBackend.h
--- a/Engine/Source/Developer/Virtualization/Private/VirtualizationSourceControlBackend.h
+++ b/Engine/Source/Developer/Virtualization/Private/VirtualizationSourceControlBackend.h
@@ -55,6 +55,11 @@ class FSemaphore;
  * SuppressNotifications[bool]:	When true the system will not display a pop up notification when a 
  *								connection error occurs, allowing the user to stay unaware of the error
  *								unless it actually causes some sort of problem. [Default=false]
+ * IgnoreFile [string]:			Sets the name of the p4 ignore file to use. When submitting payloads we 
+ *								create a custom p4 ignore file to override any ignore settings for a project
+ *								which allows us to submit from the saved directory which is normally prevented
+ *								by the default ignore file. This value can be set to what ever is used by
+ *								your perforce environment. [Default=".p4ignore.txt"]
  * 
  * Environment Variables:
  * UE-VirtualizationWorkingDir [string]:	This can be set to a valid directory path that the backend
@@ -123,6 +128,9 @@ private:
 	/** The root directory from which payloads are submitted. */
 	FString SubmissionRootDir;
 
+	/** The name to use for the p4 ignore file */
+	FString IgnoreFileName = TEXT(".p4ignore.txt");
+
 	/** Should we try to make the temp client partitioned or not? */
 	bool bUsePartitionedClient = true;
 
