Clean up some of the code documentation regarding virtualization filtering.

#rnx

diff --git a/Engine/Source/Developer/Virtualization/Private/VirtualizationManager.h b/Engine/Source/Developer/Virtualization/Private/VirtualizationManager.h
--- a/Engine/Source/Developer/Virtualization/Private/VirtualizationManager.h
+++ b/Engine/Source/Developer/Virtualization/Private/VirtualizationManager.h
@@ -49,18 +49,27 @@ struct FAnalyticsEventAttribute;
 /**
  * Filtering
  * 
- * When pushing a payload it can be filtered based on the path of the package it belongs to. The filtering options 
- * are set up via the config files. 
- * Note that this only affects pushing a payload, if the filtering for a project is changed to exclude a package that
- * is already virtualized it will still be able to pull it's payloads as needed but will store them locally in the 
- * package the next time that it is saved.
+ * By default all packages in a project can be virtualized once the system has been enabled. The can be
+ * overridden by the filtering system, either by excluding specific packages/directories, or by changing
+ * the default so that no package will be virtualized except packages/directories that have been 
+ * specifically marked as to be virtualized.
+ * 
+ * Note that the filtering is applied when a package is saved and stored as meta data in the FPackageTrailer.
+ * This means that you can scan your package files and reason about the behavior of the payloads but also
+ * means that if you change your project filtering rules that packages have to be re-saved in order for the
+ * filtering to be applied.
+ * 
  * @see ShouldVirtualizePackage or ShouldVirtualize for implementation details.
  * 
  * Basic Setup:
  * 
  * [Core.VirtualizationModule]
- * FilterMode=OptIn/OptOut					When 'OptIn' payloads will be virtualized by default, when 'OptOut' they will not be virtualized by default
- * FilterMapContent=True/False				When true any payload stored in a .umap or _BuildData.uasset file will be excluded from virtualization
+ * FilterMode=OptIn/OptOut					The general mode to be applied to all packages. With 'OptIn' packages will
+ *											not be virtualized unless their path is included via VirtualizationFilterSettings.
+ *											With 'OptOut' all packages will be virtualized unless excluded via 
+ *											VirtualizationFilterSettings [Default=OptOut]
+ * FilterMapContent=True/False				When true any payload stored in a .umap or _BuildData.uasset file will be
+ *											excluded from virtualization [Default=true]
  * 
  * PackagePath Setup:
  * 
