Fix typo in VA source control backend logging.

#rnx

diff --git a/Engine/Source/Developer/Virtualization/Private/VirtualizationSourceControlBackend.cpp b/Engine/Source/Developer/Virtualization/Private/VirtualizationSourceControlBackend.cpp
--- a/Engine/Source/Developer/Virtualization/Private/VirtualizationSourceControlBackend.cpp
+++ b/Engine/Source/Developer/Virtualization/Private/VirtualizationSourceControlBackend.cpp
@@ -879,7 +879,7 @@ bool FSourceControlBackend::TryApplySettingsFromConfigFiles(const FString& Confi
 	{
 		FParse::Value(*ConfigEntry, TEXT("Server="), ServerAddress);
 
-		if (!ClientStream.IsEmpty())
+		if (!ServerAddress.IsEmpty())
 		{
 			UE_LOG(LogVirtualization, Log, TEXT("[%s] Using server address: '%s'"), *GetDebugName(), *ServerAddress);
 		}
