Improve error message shown to the user if they fail to connect to the VA perforce backend

#rnx

- Output the entire depot path when failing to download 'payload_metainfo' along with the address of the server.
- We ask the source control provider for the status and parse the p4 port from that as there is no way (currently) to get this info directly via the source control api, which currently works but is a bit flaky long term.
-- I am not sure if we want to keep the error reporting like this and so I am loath to make changes to the source control API unless I am sure we actually want the functionality.
- Removed duplicate code paths, there is now only one branch that will log an error.
- Removed call to FConnect as that does not do anything unless a client is set and we do not use one with this connection. (FConnect is mostly about client set up validation rather than actually connecting to the server since we always make a new perforce connection for each command anyway.
-- Instead we now check if the source control provider is available after we call Init, if this returns false then no connection was made.
- Documented the purpose behind calling ISourceControlProvider::Init

diff --git a/Engine/Source/Developer/Virtualization/Private/VirtualizationSourceControlBackend.cpp b/Engine/Source/Developer/Virtualization/Private/VirtualizationSourceControlBackend.cpp
--- a/Engine/Source/Developer/Virtualization/Private/VirtualizationSourceControlBackend.cpp
+++ b/Engine/Source/Developer/Virtualization/Private/VirtualizationSourceControlBackend.cpp
@@ -179,6 +179,28 @@ private:
 	FSemaphore* Semaphore;
 };
 
+/**
+ * This utility can be used to try and find the P4PORT (server address) that the provider 
+ * is using. The source control api doesn't actually have a way to do this at the moment, 
+ * but we can get a general status text of the connection from which we can attempt to 
+ * parse the port.
+ */
+[[nodiscard]] FString GetPortFromProvider(ISourceControlProvider* SCCProvider)
+{
+	if (SCCProvider != nullptr)
+	{
+		const FString Status = SCCProvider->GetStatusText().ToString();
+
+		FString Port;
+		if (FParse::Value(*Status, TEXT("Port:"), Port))
+		{
+			return Port;
+		}		
+	}
+
+	return FString(TEXT("Unknown"));
+}
+
 /** Utility function to create a directory to submit payloads from. */
 [[nodiscard]] static bool TryCreateSubmissionSessionDirectory(FStringView SessionDirectoryPath, FStringView IgnoreFileName)
 {
@@ -330,62 +352,41 @@ IVirtualizationBackend::EConnectionStatus FSourceControlBackend::OnConnect()
 		return IVirtualizationBackend::EConnectionStatus::Error;
 	}
 
+	// Will attempt to make initial connection to the server, it doesn't return any error values we can check against but
+	// it will output problems to the log file.
 	SCCProvider->Init(true);
 
-	if (IsInGameThread())
+	if (!SCCProvider->IsAvailable())
 	{
-		// Note that if the connect is failing then we expect it to fail here rather than in the subsequent attempts to get the meta info file
-		TSharedRef<FConnect, ESPMode::ThreadSafe> ConnectCommand = ISourceControlOperation::Create<FConnect>();
-		if (SCCProvider->Execute(ConnectCommand, FString(), EConcurrency::Synchronous) != ECommandResult::Succeeded)
-		{
-			FTextBuilder Errors;
-			for (const FText& Msg : ConnectCommand->GetResultInfo().ErrorMessages)
-			{
-				Errors.AppendLine(Msg);
-			}
-
-			FText Message = FText::Format(LOCTEXT("FailedSourceControlConnection", "Failed to connect to revision control backend with the following errors:\n{0}\nThe revision control backend had trouble connecting!\nTrying logging in with the 'p4 login' command or by using p4vs/UnrealGameSync."),
-				Errors.ToText());
-
-			OnConnectionError(MoveTemp(Message));
+		OnConnectionError(FText(LOCTEXT("FailedSourceControlConnection", "Failed to connect to revision control backend, see the Message Log 'Revision Control' errors for details.\nThe revision control backend will be unable to pull payloads, is your revision control config set up correctly?")));
 
-			return IVirtualizationBackend::EConnectionStatus::Error;
-		}
+		return IVirtualizationBackend::EConnectionStatus::Error;
 	}
 
 	// When a source control depot is set up a file named 'payload_metainfo.txt' should be submitted to it's root.
 	// This allows us to check for the existence of the file to confirm that the depot root is indeed valid.
 	const FString PayloadMetaInfoPath = WriteToString<512>(DepotRoot, TEXT("payload_metainfo.txt")).ToString();
 
+	FSharedBuffer MetaInfoBuffer;
+
 #if IS_SOURCE_CONTROL_THREAD_SAFE
 	TSharedRef<FDownloadFile, ESPMode::ThreadSafe> DownloadCommand = ISourceControlOperation::Create<FDownloadFile>();
 	if (SCCProvider->Execute(DownloadCommand, PayloadMetaInfoPath, EConcurrency::Synchronous) != ECommandResult::Succeeded)
 	{
-		FText Msg = FText::Format(LOCTEXT("FailedMetaInfo", "Failed to find 'payload_metainfo.txt' in the depot '{0}'\nThe revision control backend will be unable to pull payloads, is your revision control config set up correctly?"),
-			FText::FromString(DepotRoot));
-
-		OnConnectionError(MoveTemp(Msg));
-
-		return IVirtualizationBackend::EConnectionStatus::Error;
-}
 #else
 	TSharedRef<FDownloadFile, ESPMode::ThreadSafe> DownloadCommand = ISourceControlOperation::Create<FDownloadFile>();
-	if (!SCCProvider->TryToDownloadFileFromBackgroundThread(DownloadCommand, PayloadMetaInfoPath))
+	if (SCCProvider->TryToDownloadFileFromBackgroundThread(DownloadCommand, PayloadMetaInfoPath))
 	{
-		FText Msg = FText::Format(LOCTEXT("FailedMetaInfo", "Failed to find 'payload_metainfo.txt' in the depot '{0}'\nThe revision control backend will be unable to pull payloads, is your revision control config set up correctly?"),
-			FText::FromString(DepotRoot));
-
-		OnConnectionError(MoveTemp(Msg));
+#endif //IS_SOURCE_CONTROL_THREAD_SAFE
 
-		return IVirtualizationBackend::EConnectionStatus::Error;
+		MetaInfoBuffer = DownloadCommand->GetFileData(PayloadMetaInfoPath);
 	}
-#endif //IS_SOURCE_CONTROL_THREAD_SAFE
 
-	FSharedBuffer MetaInfoBuffer = DownloadCommand->GetFileData(PayloadMetaInfoPath);
 	if (MetaInfoBuffer.IsNull())
 	{
-		FText Msg = FText::Format(LOCTEXT("FailedMetaInfo", "Failed to find 'payload_metainfo.txt' in the depot '{0}'\nThe revision control backend will be unable to pull payloads, is your revision control config set up correctly?"),
-			FText::FromString(DepotRoot));
+		FText Msg = FText::Format(LOCTEXT("FailedMetaInfo", "Failed to find '{0}' on server '{1}'\nThe revision control backend will be unable to pull payloads, is your revision control config set up correctly?"),
+			FText::FromString(PayloadMetaInfoPath),
+			FText::FromString(GetPortFromProvider(SCCProvider.Get())));
 
 		OnConnectionError(MoveTemp(Msg));
 
