Add an option ''-UseLocalIniFileSettings=True/False" to the VA source control backend, which controls if we allow the user of the locally saved 'SourceControlSettings.ini' or not

- By default the value is true and the local ini file is used.
- Allows teams to opt out of loading settings from the locally saved ini file if they wish to.

diff --git a/Engine/Source/Developer/Virtualization/Private/VirtualizationSourceControlBackend.cpp b/Engine/Source/Developer/Virtualization/Private/VirtualizationSourceControlBackend.cpp
--- a/Engine/Source/Developer/Virtualization/Private/VirtualizationSourceControlBackend.cpp
+++ b/Engine/Source/Developer/Virtualization/Private/VirtualizationSourceControlBackend.cpp
@@ -337,7 +337,10 @@ IVirtualizationBackend::EConnectionStatus FSourceControlBackend::OnConnect()
 
 	// We do not want the connection to have a client workspace so explicitly set it to empty
 	FSourceControlInitSettings SCCSettings(FSourceControlInitSettings::EBehavior::OverrideExisting);
-	SCCSettings.SetConfigBehavior(FSourceControlInitSettings::EConfigBehavior::ReadOnly);
+
+	const FSourceControlInitSettings::EConfigBehavior IniBehavior = bUseLocalIniFileSettings ?	FSourceControlInitSettings::EConfigBehavior::ReadOnly :
+																								FSourceControlInitSettings::EConfigBehavior::None;
+	SCCSettings.SetConfigBehavior(IniBehavior);
 
 	if (!ServerAddress.IsEmpty())
 	{
@@ -967,6 +970,11 @@ bool FSourceControlBackend::TryApplySettingsFromConfigFiles(const FString& Confi
 		}
 	}
 
+	{
+		FParse::Bool(*ConfigEntry, TEXT("UseLocalIniFileSettings="), bUseLocalIniFileSettings);
+		UE_LOG(LogVirtualization, Log, TEXT("[%s] Reading settings from local SourceControlSettings.ini is %s"), *GetDebugName(), bUseLocalIniFileSettings ? TEXT("enabled") : TEXT("disabled"));	
+	}
+
 	{
 		// TODO: We should just extract this from the perforce environment but that requires extending
 		// the source control api.
diff --git a/Engine/Source/Developer/Virtualization/Private/VirtualizationSourceControlBackend.h b/Engine/Source/Developer/Virtualization/Private/VirtualizationSourceControlBackend.h
--- a/Engine/Source/Developer/Virtualization/Private/VirtualizationSourceControlBackend.h
+++ b/Engine/Source/Developer/Virtualization/Private/VirtualizationSourceControlBackend.h
@@ -55,6 +55,9 @@ class FSemaphore;
  * SuppressNotifications[bool]:	When true the system will not display a pop up notification when a 
  *								connection error occurs, allowing the user to stay unaware of the error
  *								unless it actually causes some sort of problem. [Default=false]
+ * UseLocalIniFileSettings[bool]When true the revision control provider will be allowed to load connection 
+ *								settings from the users locally saved 'SourceControlSettings.ini' file, if 
+ *								false then the settings in this file (if any) will be ignored. [Default=true]
  * IgnoreFile [string]:			Sets the name of the p4 ignore file to use. When submitting payloads we 
  *								create a custom p4 ignore file to override any ignore settings for a project
  *								which allows us to submit from the saved directory which is normally prevented
@@ -148,6 +151,9 @@ private:
 
 	/** The length of time (in milliseconds) to wait after each pull attempt before retrying. */
 	int32 RetryWaitTimeMS = 100;
+
+	/** When true we allow the revision control provider to read settings from the users local SourceControlSettings.ini file */
+	bool bUseLocalIniFileSettings = true;
 };
 
 } // namespace UE::Virtualization
