Clean up VA config file parsing code

#rnx

- We no longer error if the values cannot be parsed from the config file, our defaults in BaseEngine.ini are the same defaults that we have in code anyway.
- We now trust that FConfigFile will not write to the variable passed into it on parse failure so there is no need to parse to a local variable first, then apply that to the member if the parse was a success.

diff --git a/Engine/Source/Developer/Virtualization/Private/VirtualizationManager.cpp b/Engine/Source/Developer/Virtualization/Private/VirtualizationManager.cpp
--- a/Engine/Source/Developer/Virtualization/Private/VirtualizationManager.cpp
+++ b/Engine/Source/Developer/Virtualization/Private/VirtualizationManager.cpp
@@ -208,6 +208,19 @@ private:
 	TArray<FPullRequest, TInlineAllocator<UE_INLINE_ALLOCATION_COUNT>> CurrentRequests;
 };
 
+const TCHAR* LexToString(EPackageFilterMode Value)
+{
+	switch (Value)
+	{
+	case EPackageFilterMode::OptIn:
+		return TEXT("OptIn");
+	case EPackageFilterMode::OptOut:
+		return TEXT("OptOut");
+	default:
+		checkNoEntry();
+		return TEXT("");
+	}
+}
 
 bool LexTryParseString(EPackageFilterMode& OutValue, FStringView Buffer)
 {
@@ -983,169 +996,105 @@ void FVirtualizationManager::ApplySettingsFromConfigFiles(const FConfigFile& Con
 	const TCHAR* LegacyConfigSection = TEXT("Core.ContentVirtualization");
 	const TCHAR* ConfigSection = TEXT("Core.VirtualizationModule");
 
-	// Note that all options are doubled up as we are moving the options for this module from "Core.ContentVirtualization"
+	// Note that many options are doubled up as we are moving the options for this module from "Core.ContentVirtualization"
 	// to it's own specific "Core.VirtualizationModule" section. This duplication can be removed before we ship 5.1
 	
+	// The backend graph is the most important value and so should come first for logging purposes
+	{
+		if (!ConfigFile.GetString(LegacyConfigSection, TEXT("BackendGraph"), BackendGraphName))
+		{
+			ConfigFile.GetString(ConfigSection, TEXT("BackendGraph"), BackendGraphName);
+		}
+
+		UE_LOG(LogVirtualization, Display, TEXT("\tBackendGraphName : %s"), *BackendGraphName);
+	}
+
 	{
 		// This value was moved from Core.ContentVirtualization to Core.VirtualizationModule then renamed from
 		// 'EnablePushToBackend' to 'EnablePayloadVirtualization' so there are a few paths we need to cover here.
 		// This can also be cleaned up for 5.1 shipping.
-		bool bLoadedFromFile = false;
-		bool bEnablePayloadVirtualizationFromIni = false;
-		
-		if (ConfigFile.GetBool(LegacyConfigSection, TEXT("EnablePushToBackend"), bEnablePayloadVirtualizationFromIni))
+		if (ConfigFile.GetBool(LegacyConfigSection, TEXT("EnablePushToBackend"), bAllowPackageVirtualization))
 		{
 			UE_LOG(LogVirtualization, Warning, TEXT("\tFound legacy ini file setting [Core.ContentVirtualization].EnablePushToBackend, rename to [Core.VirtualizationModule].EnablePayloadVirtualization"));
-			bLoadedFromFile = true;
 		}
-		else if (ConfigFile.GetBool(ConfigSection, TEXT("EnablePushToBackend"), bEnablePayloadVirtualizationFromIni))
+		else if (ConfigFile.GetBool(ConfigSection, TEXT("EnablePushToBackend"), bAllowPackageVirtualization))
 		{
 			UE_LOG(LogVirtualization, Warning, TEXT("\tFound legacy ini file setting [Core.VirtualizationModule].EnablePushToBackend, rename to [Core.VirtualizationModule].EnablePayloadVirtualization"));
-			bLoadedFromFile = true;
-		}
-		else if (ConfigFile.GetBool(ConfigSection, TEXT("EnablePayloadVirtualization"), bEnablePayloadVirtualizationFromIni))
-		{
-			bLoadedFromFile = true;
-		}
-
-		if (bLoadedFromFile)
-		{
-			bAllowPackageVirtualization = bEnablePayloadVirtualizationFromIni;
-			UE_LOG(LogVirtualization, Display, TEXT("\tEnablePayloadVirtualization : %s"), bAllowPackageVirtualization ? TEXT("true") : TEXT("false"));
 		}
 		else
 		{
-			UE_LOG(LogVirtualization, Error, TEXT("Failed to load [Core.VirtualizationModule].EnablePayloadVirtualization from config file!"));
+			ConfigFile.GetBool(ConfigSection, TEXT("EnablePayloadVirtualization"), bAllowPackageVirtualization);
 		}
+
+		UE_LOG(LogVirtualization, Display, TEXT("\tEnablePayloadVirtualization : %s"), bAllowPackageVirtualization ? TEXT("true") : TEXT("false"));
 	}
 
 	{
-		bool bEnableCacheAfterPullFromIni = false;
-		if (ConfigFile.GetBool(LegacyConfigSection, TEXT("EnableCacheAfterPull"), bEnableCacheAfterPullFromIni) ||
-			ConfigFile.GetBool(ConfigSection, TEXT("EnableCacheAfterPull"), bEnableCacheAfterPullFromIni))
-		{
-			UE_LOG(LogVirtualization, Warning, TEXT("\tEnableCacheAfterPull is now deprecated replace with 'EnableCacheOnPull=True|False'"));
-			if (!bEnableCacheAfterPullFromIni)
-			{
-				EnumRemoveFlags(CachingPolicy, ECachingPolicy::CacheOnPull);
-			}
-		}
-
-		bool bCacheOnPullFromIni = true;
-		if (ConfigFile.GetBool(ConfigSection, TEXT("EnableCacheOnPull"), bCacheOnPullFromIni))
+		bool bCacheOnPull = true;
+		if (ConfigFile.GetBool(LegacyConfigSection, TEXT("EnableCacheAfterPull"), bCacheOnPull) ||
+			ConfigFile.GetBool(ConfigSection, TEXT("EnableCacheAfterPull"), bCacheOnPull))
 		{
-			if (!bCacheOnPullFromIni)
-			{
-				EnumRemoveFlags(CachingPolicy, ECachingPolicy::CacheOnPull);
-			}
-
-			UE_LOG(LogVirtualization, Display, TEXT("\tEnableCacheOnPull : %s"),
-				EnumHasAllFlags(CachingPolicy, ECachingPolicy::CacheOnPull) ? TEXT("true") : TEXT("false"));
+			UE_LOG(LogVirtualization, Warning, TEXT("\tEnableCacheAfterPull is deprecated, replace with 'EnableCacheOnPull=True|False'"));
 		}
 		else
 		{
-			UE_LOG(LogVirtualization, Error, TEXT("Failed to load [Core.VirtualizationModule].EnableCacheOnPull from config file!"));
+			ConfigFile.GetBool(ConfigSection, TEXT("EnableCacheOnPull"), bCacheOnPull);
 		}
-	}
 
-	{
-		bool bCacheOnPushFromIni = true;
-		if (ConfigFile.GetBool(ConfigSection, TEXT("EnableCacheOnPush"), bCacheOnPushFromIni))
-		{
-			if (!bCacheOnPushFromIni)
-			{
-				EnumRemoveFlags(CachingPolicy, ECachingPolicy::CacheOnPush);
-			}
-			UE_LOG(LogVirtualization, Display, TEXT("\tEnableCacheOnPush : %s"), 
-				EnumHasAllFlags(CachingPolicy, ECachingPolicy::CacheOnPush) ? TEXT("true") : TEXT("false"));
-		}
-		else
+		if (!bCacheOnPull)
 		{
-			UE_LOG(LogVirtualization, Error, TEXT("Failed to load [Core.VirtualizationModule].EnableCacheOnPush from config file!"));
+			EnumRemoveFlags(CachingPolicy, ECachingPolicy::CacheOnPull);
 		}
-	}
 
-	int64 MinPayloadLengthFromIni = 0;
-	if (ConfigFile.GetInt64(LegacyConfigSection, TEXT("MinPayloadLength"), MinPayloadLengthFromIni) ||
-		ConfigFile.GetInt64(ConfigSection, TEXT("MinPayloadLength"), MinPayloadLengthFromIni))
-	{
-		MinPayloadLength = MinPayloadLengthFromIni;
-		UE_LOG(LogVirtualization, Display, TEXT("\tMinPayloadLength : %" INT64_FMT), MinPayloadLength );
-	}
-	else
-	{
-		UE_LOG(LogVirtualization, Error, TEXT("Failed to load [Core.VirtualizationModule].MinPayloadLength from config file!"));
+		UE_LOG(LogVirtualization, Display, TEXT("\tEnableCacheOnPull : %s"), EnumHasAllFlags(CachingPolicy, ECachingPolicy::CacheOnPull) ? TEXT("true") : TEXT("false"));
 	}
 
-	FString BackendGraphNameFromIni;
-	if (ConfigFile.GetString(LegacyConfigSection, TEXT("BackendGraph"), BackendGraphNameFromIni) ||
-		ConfigFile.GetString(ConfigSection, TEXT("BackendGraph"), BackendGraphNameFromIni))
-	{
-		BackendGraphName = BackendGraphNameFromIni;
-		UE_LOG(LogVirtualization, Display, TEXT("\tBackendGraphName : %s"), *BackendGraphName );
-	}
-	else
 	{
-		UE_LOG(LogVirtualization, Error, TEXT("Failed to load [Core.VirtualizationModule].BackendGraph from config file!"));
-	}
+		bool bCacheOnPush = true;
+		ConfigFile.GetBool(ConfigSection, TEXT("EnableCacheOnPush"), bCacheOnPush);
 
-	FString VirtualizationProcessTagFromIni;
-	if (ConfigFile.GetString(ConfigSection, TEXT("VirtualizationProcessTag"), VirtualizationProcessTagFromIni))
-	{
-		VirtualizationProcessTag = VirtualizationProcessTagFromIni;
-		UE_LOG(LogVirtualization, Display, TEXT("\tVirtualizationProcessTag : %s"), *VirtualizationProcessTag);
-	}
-	else
-	{
-		UE_LOG(LogVirtualization, Error, TEXT("Failed to load [Core.VirtualizationModule].VirtualizationProcessTag from config file!"));
+		if (!bCacheOnPush)
+		{
+			EnumRemoveFlags(CachingPolicy, ECachingPolicy::CacheOnPush);
+		}
+		
+		UE_LOG(LogVirtualization, Display, TEXT("\tEnableCacheOnPush : %s"),  EnumHasAllFlags(CachingPolicy, ECachingPolicy::CacheOnPush) ? TEXT("true") : TEXT("false"));
 	}
 
-	FString FilterModeFromIni;
-	if (ConfigFile.GetString(LegacyConfigSection, TEXT("FilterMode"), FilterModeFromIni) ||
-		ConfigFile.GetString(ConfigSection, TEXT("FilterMode"), FilterModeFromIni))
 	{
-		if(LexTryParseString(FilteringMode, FilterModeFromIni))
-		{
-			UE_LOG(LogVirtualization, Display, TEXT("\tFilterMode : %s"), *FilterModeFromIni);
-		}
-		else
+		if (!ConfigFile.GetInt64(LegacyConfigSection, TEXT("MinPayloadLength"), MinPayloadLength))
 		{
-			UE_LOG(LogVirtualization, Error, TEXT("[Core.VirtualizationModule].FilterMode was an invalid value! Allowed: 'OptIn'|'OptOut' Found '%s'"), *FilterModeFromIni);
+			ConfigFile.GetInt64(ConfigSection, TEXT("MinPayloadLength"), MinPayloadLength);
 		}
+
+		UE_LOG(LogVirtualization, Display, TEXT("\tMinPayloadLength : %" INT64_FMT), MinPayloadLength);
 	}
-	else
+
 	{
-		UE_LOG(LogVirtualization, Error, TEXT("Failed to load [Core.VirtualizationModule]FilterMode from config file!"));
+		ConfigFile.GetString(ConfigSection, TEXT("VirtualizationProcessTag"), VirtualizationProcessTag);
+		UE_LOG(LogVirtualization, Display, TEXT("\tVirtualizationProcessTag : %s"), *VirtualizationProcessTag);
 	}
 
-	// Deprecated
 	{
-		bool bDummyValue = true;
-		if (ConfigFile.GetBool(LegacyConfigSection, TEXT("FilterEngineContent"), bDummyValue) ||
-			ConfigFile.GetBool(ConfigSection, TEXT("FilterEngineContent"), bDummyValue))
+		FString FilterModeString;
+		if (ConfigFile.GetString(LegacyConfigSection, TEXT("FilterMode"), FilterModeString) ||
+			ConfigFile.GetString(ConfigSection, TEXT("FilterMode"), FilterModeString))
 		{
-			UE_LOG(LogVirtualization, Warning, TEXT("\tFilterEngineContent is now deprecated (engine content is never virtualized)"));
+			if (!LexTryParseString(FilteringMode, FilterModeString))
+			{
+				UE_LOG(LogVirtualization, Error, TEXT("[Core.VirtualizationModule].FilterMode was an invalid value! Allowed: 'OptIn'|'OptOut' Found '%s'"), *FilterModeString);
+			}
 		}
 
-		if (ConfigFile.GetBool(LegacyConfigSection, TEXT("FilterEnginePluginContent"), bDummyValue) ||
-			ConfigFile.GetBool(ConfigSection, TEXT("FilterEnginePluginContent"), bDummyValue))
-		{
-			UE_LOG(LogVirtualization, Warning, TEXT("\tFilterEnginePluginContent is now deprecated (engine content is never virtualized)"));
-		}
+		UE_LOG(LogVirtualization, Display, TEXT("\tFilterMode : %s"), LexToString(FilteringMode));
 	}
 
-	bool bFilterMapContentFromIni = false;
-	if (ConfigFile.GetBool(ConfigSection, TEXT("FilterMapContent"), bFilterMapContentFromIni))
 	{
-		bFilterMapContent = bFilterMapContentFromIni;
+		ConfigFile.GetBool(ConfigSection, TEXT("FilterMapContent"), bFilterMapContent);
 		UE_LOG(LogVirtualization, Display, TEXT("\tFilterMapContent : %s"), bFilterMapContent ? TEXT("true") : TEXT("false"));
 	}
-	else
-	{
-		UE_LOG(LogVirtualization, Error, TEXT("Failed to load [Core.VirtualizationModule].FilterMapContent from config file!"));
-	}
 
-	// Optional
+
 	TArray<FString> DisabledAssetTypesFromIni;
 	if (ConfigFile.GetArray(LegacyConfigSection, TEXT("DisabledAsset"), DisabledAssetTypesFromIni) > 0 ||
 		ConfigFile.GetArray(ConfigSection, TEXT("DisabledAsset"), DisabledAssetTypesFromIni) > 0)
@@ -1159,33 +1108,35 @@ void FVirtualizationManager::ApplySettingsFromConfigFiles(const FConfigFile& Con
 		}	
 	}
 
-	bool bAllowSubmitIfVirtualizationFailedFromIni = true;
-	if (ConfigFile.GetBool(ConfigSection, TEXT("AllowSubmitIfVirtualizationFailed"), bAllowSubmitIfVirtualizationFailedFromIni))
 	{
-		bAllowSubmitIfVirtualizationFailed = bAllowSubmitIfVirtualizationFailedFromIni;
+		ConfigFile.GetBool(ConfigSection, TEXT("AllowSubmitIfVirtualizationFailed"), bAllowSubmitIfVirtualizationFailed);
 		UE_LOG(LogVirtualization, Display, TEXT("\tAllowSubmitIfVirtualizationFailed : %s"), bAllowSubmitIfVirtualizationFailed ? TEXT("true") : TEXT("false"));
 	}
-	else
-	{
-		UE_LOG(LogVirtualization, Error, TEXT("Failed to load [Core.VirtualizationModule].AllowSubmitIfVirtualizationFailed from config file!"));
-	}
 
 #if UE_VIRTUALIZATION_CONNECTION_LAZY_INIT == 0
-	bool bLazyInitConnectionsFromIni = true;
-	if (ConfigFile.GetBool(ConfigSection, TEXT("LazyInitConnections"), bLazyInitConnectionsFromIni))
-	{
-		bLazyInitConnections = bLazyInitConnectionsFromIni;
-		UE_LOG(LogVirtualization, Display, TEXT("\tLazyInitConnections : %s"), bLazyInitConnections ? TEXT("true") : TEXT("false"));
-	}
-	else
-	{
-		UE_LOG(LogVirtualization, Error, TEXT("Failed to load [Core.VirtualizationModule].LazyInitConnections from config file!"));
-	}
+	ConfigFile.GetBool(ConfigSection, TEXT("LazyInitConnections"), bLazyInitConnections);
+	UE_LOG(LogVirtualization, Display, TEXT("\tLazyInitConnections : %s"), bLazyInitConnections ? TEXT("true") : TEXT("false"));
 #else
 	bLazyInitConnections = true;
-	UE_LOG(LogVirtualization, Display, TEXT("\tLazyInitConnections : %s (set by code)"), bLazyInitConnections ? TEXT("true") : TEXT("false"));
+	UE_LOG(LogVirtualization, Display, TEXT("\tLazyInitConnections : true (set by code)"));
 #endif //UE_VIRTUALIZATION_CONNECTION_LAZY_INIT
 
+	// Deprecated
+	{
+		bool bDummyValue = true;
+		if (ConfigFile.GetBool(LegacyConfigSection, TEXT("FilterEngineContent"), bDummyValue) ||
+			ConfigFile.GetBool(ConfigSection, TEXT("FilterEngineContent"), bDummyValue))
+		{
+			UE_LOG(LogVirtualization, Warning, TEXT("\tFilterEngineContent is now deprecated (engine content is never virtualized)"));
+		}
+
+		if (ConfigFile.GetBool(LegacyConfigSection, TEXT("FilterEnginePluginContent"), bDummyValue) ||
+			ConfigFile.GetBool(ConfigSection, TEXT("FilterEnginePluginContent"), bDummyValue))
+		{
+			UE_LOG(LogVirtualization, Warning, TEXT("\tFilterEnginePluginContent is now deprecated (engine content is never virtualized)"));
+		}
+	}
+
 	// Check for any legacy settings and print them out (easier to do this in one block rather than one and time)
 	{
 		// Entries that are allows to be in [Core.ContentVirtualization
