Add a new error handling flow (off by default) to VA when payloads fail to pull.

- The new flow will display an error dialog to the user informing them of the problem and ask if they want to retry or quit the process. This is quite a shift if how we have been handling failed payload pulls up until and this is the first rough implementation so for now it is off by default.
- We only show one dialog at a time and if a new error occurs while we are waiting for user input, that thread is set to wait on the existing dialogs outcome.

diff --git a/Engine/Config/BaseEngine.ini b/Engine/Config/BaseEngine.ini
--- a/Engine/Config/BaseEngine.ini
+++ b/Engine/Config/BaseEngine.ini
@@ -1346,6 +1346,7 @@ AllowSubmitIfVirtualizationFailed=false
 FilterMode=OptOut
 FilterMapContent=True
 LazyInitConnections=False
+UseLegacyErrorHandling=True
 ; The following assets should not virtualize until it is confirmed that they do not crash when virtualized payloads are missing
 +DisabledAsset=StaticMesh
 +DisabledAsset=SoundWave
diff --git a/Engine/Source/Developer/Virtualization/Private/VirtualizationManager.cpp b/Engine/Source/Developer/Virtualization/Private/VirtualizationManager.cpp
--- a/Engine/Source/Developer/Virtualization/Private/VirtualizationManager.cpp
+++ b/Engine/Source/Developer/Virtualization/Private/VirtualizationManager.cpp
@@ -20,6 +20,8 @@
 #include "VirtualizationFilterSettings.h"
 #include "AnalyticsEventAttribute.h"
 
+#include "Misc/MessageDialog.h"
+
 #define LOCTEXT_NAMESPACE "Virtualization"
 
 
@@ -195,10 +197,10 @@ public:
 		return CurrentRequests;
 	}
 
-	/** Returns if we still have payloads that have not yet been found */
-	bool HasRemainingRequests() const
+	/** Returns if there are still requests that need servicing or not */
+	bool IsWorkComplete() const
 	{
-		return !CurrentRequests.IsEmpty();
+		return CurrentRequests.IsEmpty();
 	}
 
 private:
@@ -490,6 +492,7 @@ FVirtualizationManager::FVirtualizationManager()
 	, bFilterMapContent(true)
 	, bAllowSubmitIfVirtualizationFailed(false)
 	, bLazyInitConnections(false)
+	, bUseLegacyErrorHandling(true)
 	, bPendingBackendConnections(false)
 {
 }
@@ -1121,6 +1124,11 @@ void FVirtualizationManager::ApplySettingsFromConfigFiles(const FConfigFile& Con
 	UE_LOG(LogVirtualization, Display, TEXT("\tLazyInitConnections : true (set by code)"));
 #endif //UE_VIRTUALIZATION_CONNECTION_LAZY_INIT
 
+	{
+		ConfigFile.GetBool(ConfigSection, TEXT("UseLegacyErrorHandling"), bUseLegacyErrorHandling);
+		UE_LOG(LogVirtualization, Display, TEXT("\tUseLegacyErrorHandling : %s"), bUseLegacyErrorHandling ? TEXT("true") : TEXT("false"));
+	}
+
 	// Deprecated
 	{
 		bool bDummyValue = true;
@@ -1826,36 +1834,49 @@ void FVirtualizationManager::PullDataFromAllBackends(TArrayView<FPullRequest> Re
 	}
 
 	FPullRequestCollection RequestsCollection(Requests);
-
-	for (IVirtualizationBackend* Backend : PullEnabledBackends)
+	
+	while(true)
 	{
-		check(Backend != nullptr);
-
-		if (Backend->IsOperationDebugDisabled(IVirtualizationBackend::EOperations::Pull))
+		for (IVirtualizationBackend* Backend : PullEnabledBackends)
 		{
-			UE_LOG(LogVirtualization, Verbose, TEXT("Pulling from backend '%s' is debug disabled"), *Backend->GetDebugName());
-			continue;
-		}
+			check(Backend != nullptr);
 
-		if (Backend->GetConnectionStatus() != IVirtualizationBackend::EConnectionStatus::Connected)
-		{
-			UE_LOG(LogVirtualization, Verbose, TEXT("Cannot pull from backend '%s' as it is not connected"), *Backend->GetDebugName());
-			continue;
-		}
+			if (Backend->IsOperationDebugDisabled(IVirtualizationBackend::EOperations::Pull))
+			{
+				UE_LOG(LogVirtualization, Verbose, TEXT("Pulling from backend '%s' is debug disabled"), *Backend->GetDebugName());
+				continue;
+			}
+
+			if (Backend->GetConnectionStatus() != IVirtualizationBackend::EConnectionStatus::Connected)
+			{
+				UE_LOG(LogVirtualization, Verbose, TEXT("Cannot pull from backend '%s' as it is not connected"), *Backend->GetDebugName());
+				continue;
+			}
 
-		PullDataFromBackend(*Backend, RequestsCollection.GetRequests());
+			PullDataFromBackend(*Backend, RequestsCollection.GetRequests());
 
-		const bool bShouldCache = EnumHasAllFlags(CachingPolicy, ECachingPolicy::CacheOnPull);
+			const bool bShouldCache = EnumHasAllFlags(CachingPolicy, ECachingPolicy::CacheOnPull);
 
-		TArray<FPushRequest> PayloadsToCache = RequestsCollection.OnPullCompleted(*Backend, bShouldCache);
-		if (!PayloadsToCache.IsEmpty())
-		{
-			CachePayloads(PayloadsToCache, Backend);
+			TArray<FPushRequest> PayloadsToCache = RequestsCollection.OnPullCompleted(*Backend, bShouldCache);
+			if (!PayloadsToCache.IsEmpty())
+			{
+				CachePayloads(PayloadsToCache, Backend);
+			}
+
+			// We can early out if there is no more requests to make
+			if (RequestsCollection.IsWorkComplete())
+			{
+				break;
+			}
 		}
 
-		if (!RequestsCollection.HasRemainingRequests())
+		if (RequestsCollection.IsWorkComplete())
 		{
-			break;
+			return; // All payloads pulled
+		}
+		else if (OnPayloadPullError() != ErrorHandlingResult::Retry)
+		{
+			return; // Some payloads failed to pull
 		}
 	}
 }
@@ -1884,6 +1905,43 @@ void FVirtualizationManager::PullDataFromBackend(IVirtualizationBackend& Backend
 #endif //ENABLE_COOK_STATS
 }
 
+FVirtualizationManager::ErrorHandlingResult FVirtualizationManager::OnPayloadPullError()
+{
+	if (bUseLegacyErrorHandling)
+	{
+		return ErrorHandlingResult::AcceptFailedPayloads;
+	}
+
+	static FCriticalSection CriticalSection;
+
+	if (CriticalSection.TryLock())
+	{
+		FText Title = FText::FromString(TEXT("Failed to pull virtualized data!"));
+
+
+		FString Msg = FString::Printf(	TEXT("Failed to pull payload(s) from virtualization storage and allowing the editor to continue could corrupt data!"
+											 "\n\n[Yes] Retry pulling the data\n[No] Quit the editor"));
+
+		FText Message = FText::FromString(Msg);
+		EAppReturnType::Type Result = FMessageDialog::Open(EAppMsgType::YesNo, EAppReturnType::No, Message, &Title);
+
+		if (Result == EAppReturnType::No)
+		{
+			UE_LOG(LogVirtualization, Error, TEXT("Failed to pull payloads from persistent storage and the connection could not be re-established, exiting..."));
+			GIsCriticalError = 1;
+			FPlatformMisc::RequestExit(true);
+		}
+
+		CriticalSection.Unlock();
+	}
+	else
+	{
+		FScopeLock _(&CriticalSection);
+	}
+
+	return ErrorHandlingResult::Retry;
+}
+
 bool FVirtualizationManager::ShouldVirtualizeAsset(const UObject* OwnerObject) const
 {
 	if (OwnerObject == nullptr)
diff --git a/Engine/Source/Developer/Virtualization/Private/VirtualizationManager.h b/Engine/Source/Developer/Virtualization/Private/VirtualizationManager.h
--- a/Engine/Source/Developer/Virtualization/Private/VirtualizationManager.h
+++ b/Engine/Source/Developer/Virtualization/Private/VirtualizationManager.h
@@ -115,6 +115,10 @@ struct FAnalyticsEventAttribute;
  *												This can remove lengthy connection steps from the process init phase and then only connect
  *												if we actually need that service. Note that if this is true then the connection can come from
  *												any thread, so custom backend code will need to take that into account. [Default=false]
+ * UseLegacyErrorHandling [bool]:				Controls how we deal with errors encountered when pulling payloads. When true a failed payload 
+ *												pull will return an error and allow the process to carry on (the original error handling logic)
+ *												and when false a dialog will be displayed to the user warning them about the failed pull and 
+ *												prompting them to retry the pull or to quit the process. [Default=true]
  */
 
 namespace UE::Virtualization
@@ -214,6 +218,16 @@ private:
 
 	void PullDataFromAllBackends(TArrayView<FPullRequest> Requests);
 	void PullDataFromBackend(IVirtualizationBackend& Backend, TArrayView<FPullRequest> Requests);
+
+	enum class ErrorHandlingResult
+	{
+		/** We should try to pull the failed payloads again */
+		Retry = 0,
+		/** We should accept that some of the payloads failed and leave it to the calling system to handle */
+		AcceptFailedPayloads
+	};
+
+	ErrorHandlingResult OnPayloadPullError();
 	
 	bool ShouldVirtualizeAsset(const UObject* Owner) const;
 
@@ -290,6 +304,9 @@ private:
 	/** Should backends defer connecting to their services until first use */
 	bool bLazyInitConnections;
 
+	/** When true we do not display an error dialog on failed payload pulling and rely on the caller handling it */
+	bool bUseLegacyErrorHandling;
+
 private:
 
 	/** The name of the current project */
