Allow the error message displayed when VA payloads fail to pull to be customized.

- The payload pull error dialog now suports an additional custom line in the dialog that can be set via Engine:[Core.VirtualizationModule]:AdditionalPullErrorMsg. The intent is to allow teams to provide additional help and info to their users.
- The default error text is now localized.

diff --git a/Engine/Config/BaseEngine.ini b/Engine/Config/BaseEngine.ini
--- a/Engine/Config/BaseEngine.ini
+++ b/Engine/Config/BaseEngine.ini
@@ -1347,6 +1347,7 @@ FilterMode=OptOut
 FilterMapContent=True
 LazyInitConnections=False
 UseLegacyErrorHandling=True
+PullErrorAdditionalMsg=""
 ; The following assets should not virtualize until it is confirmed that they do not crash when virtualized payloads are missing
 +DisabledAsset=StaticMesh
 +DisabledAsset=SoundWave
diff --git a/Engine/Source/Developer/Virtualization/Private/VirtualizationManager.cpp b/Engine/Source/Developer/Virtualization/Private/VirtualizationManager.cpp
--- a/Engine/Source/Developer/Virtualization/Private/VirtualizationManager.cpp
+++ b/Engine/Source/Developer/Virtualization/Private/VirtualizationManager.cpp
@@ -1129,6 +1129,11 @@ void FVirtualizationManager::ApplySettingsFromConfigFiles(const FConfigFile& Con
 		UE_LOG(LogVirtualization, Display, TEXT("\tUseLegacyErrorHandling : %s"), bUseLegacyErrorHandling ? TEXT("true") : TEXT("false"));
 	}
 
+	{
+		ConfigFile.GetString(ConfigSection, TEXT("PullErrorAdditionalMsg"), PullErrorAdditionalMsg);
+		// This value is not echoed to the log file, as seeing an error string there might confuse users
+	}
+
 	// Deprecated
 	{
 		bool bDummyValue = true;
@@ -1916,13 +1921,24 @@ FVirtualizationManager::ErrorHandlingResult FVirtualizationManager::OnPayloadPul
 
 	if (CriticalSection.TryLock())
 	{
-		FText Title = FText::FromString(TEXT("Failed to pull virtualized data!"));
+		
+		const FText Title(LOCTEXT("VAPullTitle", "Failed to pull virtualized data!"));
+
+		FTextBuilder MsgBuilder;
+		MsgBuilder.AppendLine(LOCTEXT("VAPullMsgHeader", "Failed to pull payload(s) from virtualization storage and allowing the editor to continue could corrupt data!"));
+		
+		if (!PullErrorAdditionalMsg.IsEmpty())
+		{
+			MsgBuilder.AppendLine(FString(TEXT("")));
+			MsgBuilder.AppendLine(PullErrorAdditionalMsg);
+		}
 
+		MsgBuilder.AppendLine(FString(TEXT("")));
+		MsgBuilder.AppendLine(LOCTEXT("VAPullMsgYes", "[Yes] Retry pulling the data"));
+		MsgBuilder.AppendLine(LOCTEXT("VAPullMsgNo", "[No] Quit the editor"));
 
-		FString Msg = FString::Printf(	TEXT("Failed to pull payload(s) from virtualization storage and allowing the editor to continue could corrupt data!"
-											 "\n\n[Yes] Retry pulling the data\n[No] Quit the editor"));
+		const FText Message = MsgBuilder.ToText();
 
-		FText Message = FText::FromString(Msg);
 		EAppReturnType::Type Result = FMessageDialog::Open(EAppMsgType::YesNo, EAppReturnType::No, Message, &Title);
 
 		if (Result == EAppReturnType::No)
diff --git a/Engine/Source/Developer/Virtualization/Private/VirtualizationManager.h b/Engine/Source/Developer/Virtualization/Private/VirtualizationManager.h
--- a/Engine/Source/Developer/Virtualization/Private/VirtualizationManager.h
+++ b/Engine/Source/Developer/Virtualization/Private/VirtualizationManager.h
@@ -307,6 +307,9 @@ private:
 	/** When true we do not display an error dialog on failed payload pulling and rely on the caller handling it */
 	bool bUseLegacyErrorHandling;
 
+	/** An additional error message to display when pulling payloads fails */
+	FString PullErrorAdditionalMsg;
+
 private:
 
 	/** The name of the current project */
