Add code documentation describing the perforce set up required for storing VA payloads.

#rnx

diff --git a/Engine/Source/Developer/Virtualization/Private/VirtualizationSourceControlBackend.h b/Engine/Source/Developer/Virtualization/Private/VirtualizationSourceControlBackend.h
--- a/Engine/Source/Developer/Virtualization/Private/VirtualizationSourceControlBackend.h
+++ b/Engine/Source/Developer/Virtualization/Private/VirtualizationSourceControlBackend.h
@@ -19,7 +19,8 @@ class FSemaphore;
  * a binary blob.
  * It is assumed that the files are stored with the same path convention as the file system
  * backend, found in Utils::PayloadIdToPath.
- *
+ * ----------------------------------------------------------------------------
+ * 
  * Ini file setup:
  * 'Name'=(Type=P4SourceControl, DepotRoot="//ABC/")
  * Where 'Name' is the backend name in the hierarchy and 'ABC' is the path in the source control
@@ -69,8 +70,28 @@ class FSemaphore;
  *											should use as the root location to submit payloads from.
  *											If the users machine has this set then 'SubmitFromTempDir' 
  *											will be ignored. 
+ * ----------------------------------------------------------------------------
+ * 
+ * Perforce Storage  Setup:
+ * 
+ * Once you have decided where in perforce you want to store the virtualized payloads you will need to submit
+ * a dummy file to the root of that location named "payload_metainfo.txt" preferably with a short message describing
+ * the storage locations purpose.
+ * When initialized this backend will first attempt to find and download this dummy file as a way of validating that
+ * a) the process can correctly connect to perforce 
+ * b) the process has the correct location setup in the config file.
+ * This validation not only allows us to alert the user to misconfiguration issues before they need to access the data
+ * but allows us to prevent users from accidentally pushing virtualized payloads to the wrong location.
+ * 
+ * So for example, if the virtualized payloads are to be stored under '//payloads/project/...' then the backend would
+ * expect to find a '//payloads/project/payload_metainfo.txt' file that it would download via a p4 print command.
+ * 
+ * If this file is not found then the backend will log an error the first time that it attempts to connect and will prevent
+ * the user from pulling or pushing payloads until the issue is resolved.
+ * ----------------------------------------------------------------------------
+ * 
+ * Perforce Settings Setup:
  * 
- * Perforce Setup
  * The backend will attempt to use the current global environment settings for perforce to determine 
  * the P4PORT and P4USER values. If you wish to set them via an ini file then they can be done per
  * user by adding the following to Saved\Config\<platform>Editor\SourceControlSettings.ini
