Do not assume that all VA backends will be inside of a UE::Virtualization namespace

#rnx

- The UE_REGISTER_VIRTUALIZATION_BACKEND_FACTORY macro would only compile when inside of the UE::Virtualization namespace, which is true for all of the default implementations that we provide but will probably not be true for any extension backends that third parties might write.

diff --git a/Engine/Source/Developer/Virtualization/Public/IVirtualizationBackend.h b/Engine/Source/Developer/Virtualization/Public/IVirtualizationBackend.h
--- a/Engine/Source/Developer/Virtualization/Public/IVirtualizationBackend.h
+++ b/Engine/Source/Developer/Virtualization/Public/IVirtualizationBackend.h
@@ -278,13 +278,13 @@ ENUM_CLASS_FLAGS(IVirtualizationBackend::EOperations);
  * @param The name used in config ini files to reference this backend type.
  */
 #define UE_REGISTER_VIRTUALIZATION_BACKEND_FACTORY(BackendClass, ConfigName) \
-	class F##BackendClass##Factory : public IVirtualizationBackendFactory \
+	class F##BackendClass##Factory : public UE::Virtualization::IVirtualizationBackendFactory \
 	{ \
 	public: \
 		F##BackendClass##Factory() { IModularFeatures::Get().RegisterModularFeature(FName("VirtualizationBackendFactory"), this); }\
 		virtual ~F##BackendClass##Factory() { IModularFeatures::Get().UnregisterModularFeature(FName("VirtualizationBackendFactory"), this); } \
 	private: \
-		virtual TUniquePtr<IVirtualizationBackend> CreateInstance(FStringView ProjectName, FStringView ConfigName) override \
+		virtual TUniquePtr<UE::Virtualization::IVirtualizationBackend> CreateInstance(FStringView ProjectName, FStringView ConfigName) override \
 		{ \
 			return MakeUnique<BackendClass>(ProjectName, ConfigName, WriteToString<256>(#ConfigName, TEXT(" - "), ConfigName).ToString()); \
 		} \
@@ -293,7 +293,7 @@ ENUM_CLASS_FLAGS(IVirtualizationBackend::EOperations);
 	static F##BackendClass##Factory BackendClass##Factory##Instance;
 
 #define UE_REGISTER_VIRTUALIZATION_BACKEND_FACTORY_LEGACY_IMPL(FactoryName, BackendClass, LegacyConfigName, ConfigName) \
-	class F##FactoryName : public IVirtualizationBackendFactory \
+	class F##FactoryName : public UE::Virtualization::IVirtualizationBackendFactory \
 	{ \
 	public: \
 		F##FactoryName(const TCHAR* InLegacyConfigName, const TCHAR* InNewConfigName) \
@@ -302,7 +302,7 @@ ENUM_CLASS_FLAGS(IVirtualizationBackend::EOperations);
 			{ IModularFeatures::Get().RegisterModularFeature(FName("VirtualizationBackendFactory"), this); }\
 		virtual ~F##FactoryName() { IModularFeatures::Get().UnregisterModularFeature(FName("VirtualizationBackendFactory"), this); } \
 	private: \
-		virtual TUniquePtr<IVirtualizationBackend> CreateInstance(FStringView ProjectName, FStringView ConfigName) override \
+		virtual TUniquePtr<UE::Virtualization::IVirtualizationBackend> CreateInstance(FStringView ProjectName, FStringView ConfigName) override \
 		{ \
 			UE_LOG(LogVirtualization, Warning, TEXT("Creating a backend via the legacy config name '%s' use '%s' instead"), *StoredLegacyConfigName, *StoredNewConfigName); \
 			return MakeUnique<BackendClass>(ProjectName, ConfigName, WriteToString<256>(#ConfigName, TEXT(" - "), ConfigName).ToString()); \
