Improve the logging during the virtualization process.

#rnx

diff --git a/Engine/Source/Developer/Virtualization/Private/PackageVirtualizationProcess.cpp b/Engine/Source/Developer/Virtualization/Private/PackageVirtualizationProcess.cpp
--- a/Engine/Source/Developer/Virtualization/Private/PackageVirtualizationProcess.cpp
+++ b/Engine/Source/Developer/Virtualization/Private/PackageVirtualizationProcess.cpp
@@ -281,6 +281,8 @@ void VirtualizePackages(TConstArrayView<FString> PackagePaths, EVirtualizationOp
 	Progress.EnterProgressFrame(1.0f);
 	if(System.IsPushingEnabled(EStorageType::Cache))
 	{
+		UE_LOG(LogVirtualization, Display, TEXT("Pushing payload(s) to EStorageType::Cache storage..."));
+
 		if (!System.PushData(PayloadsToSubmit, EStorageType::Cache))
 		{
 			// Caching is not critical to the process so we only warn if it fails
@@ -307,6 +309,8 @@ void VirtualizePackages(TConstArrayView<FString> PackagePaths, EVirtualizationOp
 	{
 		Progress.EnterProgressFrame(1.0f);
 
+		UE_LOG(LogVirtualization, Display, TEXT("Pushing payload(s) to EStorageType::Persistent storage..."));
+
 		if (!System.PushData(PayloadsToSubmit, EStorageType::Persistent))
 		{
 			FText Message = LOCTEXT("Virtualization_PushFailure", "Failed to push payloads");
@@ -353,28 +357,31 @@ void VirtualizePackages(TConstArrayView<FString> PackagePaths, EVirtualizationOp
 	TArray<TPair<FPackagePath, FString>> PackagesToReplace;
 
 	// Any package with an updated trailer needs to be copied and an updated trailer appended
-	for (FPackageInfo& PackageInfo : Packages)
 	{
-		if (!PackageInfo.bWasTrailerUpdated)
+		UE_LOG(LogVirtualization, Display, TEXT("Creating packages with the virtualized data removed..."));
+		for (FPackageInfo& PackageInfo : Packages)
 		{
-			continue;
-		}
+			if (!PackageInfo.bWasTrailerUpdated)
+			{
+				continue;
+			}
 
-		const FPackagePath& PackagePath = PackageInfo.Path; // No need to validate path, we checked this earlier
+			const FPackagePath& PackagePath = PackageInfo.Path; // No need to validate path, we checked this earlier
 
-		FString NewPackagePath = DuplicatePackageWithUpdatedTrailer(PackagePath.GetLocalFullPath(), PackageInfo.Trailer, OutResultInfo.Errors);
+			FString NewPackagePath = DuplicatePackageWithUpdatedTrailer(PackagePath.GetLocalFullPath(), PackageInfo.Trailer, OutResultInfo.Errors);
 
-		if (!NewPackagePath.IsEmpty())
-		{
-			// Now that we have successfully created a new version of the package with an updated trailer 
-			// we need to mark that it should replace the original package.
-			PackagesToReplace.Emplace(PackagePath, MoveTemp(NewPackagePath));
-		}
-		else
-		{
-			return;
-		}
+			if (!NewPackagePath.IsEmpty())
+			{
+				// Now that we have successfully created a new version of the package with an updated trailer 
+				// we need to mark that it should replace the original package.
+				PackagesToReplace.Emplace(PackagePath, MoveTemp(NewPackagePath));
+			}
+			else
+			{
+				return;
+			}
 
+		}
 	}
 
 	UE_LOG(LogVirtualization, Display, TEXT("%d package(s) had their trailer container modified and need to be updated"), PackagesToReplace.Num());
@@ -384,14 +391,15 @@ void VirtualizePackages(TConstArrayView<FString> PackagePaths, EVirtualizationOp
 		// TODO: Consider using the SavePackage model (move the original, then replace, so we can restore all of the original packages if needed)
 		// having said that, once a package is in PackagesToReplace it should still be safe to submit so maybe we don't need this level of protection?
 
-		
+		UE_CLOG(!PackagesToReplace.IsEmpty(), LogVirtualization, Display, TEXT("Detaching loaded packages from disk..."));
+
 		// We need to reset the loader of any package that we want to re-save over
 		for (const TPair<FPackagePath, FString>& Pair : PackagesToReplace)
 		{
 			UPackage* Package = FindObjectFast<UPackage>(nullptr, Pair.Key.GetPackageFName());
 			if (Package != nullptr)
 			{
-				UE_LOG(LogVirtualization, Verbose, TEXT("Detaching '%s' from disk so that it can be virtualized"), *Pair.Key.GetDebugName());
+				UE_LOG(LogVirtualization, Verbose, TEXT("Detaching '%s'"), *Pair.Key.GetDebugName());
 				ResetLoadersForSave(Package, *Pair.Key.GetLocalFullPath());
 			}
 		}
@@ -399,6 +407,8 @@ void VirtualizePackages(TConstArrayView<FString> PackagePaths, EVirtualizationOp
 		// Should we try to check out packages from revision control?
 		if (EnumHasAnyFlags(Options, EVirtualizationOptions::Checkout))
 		{
+			UE_CLOG(!PackagesToReplace.IsEmpty(), LogVirtualization, Display, TEXT("Checking out packages from revision control..."));
+
 			TArray<FString> FilesToCheckState;
 			FilesToCheckState.Reserve(PackagesToReplace.Num());
 
@@ -413,6 +423,8 @@ void VirtualizePackages(TConstArrayView<FString> PackagePaths, EVirtualizationOp
 			}
 		}
 
+		UE_CLOG(!PackagesToReplace.IsEmpty(), LogVirtualization, Display, TEXT("Checking packages for write access permission..."));
+
 		// Now check to see if there are package files that cannot be edited because they are read only
 		for (int32 Index = 0; Index < PackagesToReplace.Num(); ++Index)
 		{
@@ -434,6 +446,8 @@ void VirtualizePackages(TConstArrayView<FString> PackagePaths, EVirtualizationOp
 			}
 		}
 
+		UE_CLOG(!PackagesToReplace.IsEmpty(), LogVirtualization, Display, TEXT("Replacing old packages with the virtualized version..."));
+
 		// Since we had no errors we can now replace all of the packages that were virtualized data with the virtualized replacement file.
 		for(const TPair<FPackagePath,FString>&  Iterator : PackagesToReplace)
 		{
@@ -455,6 +469,7 @@ void VirtualizePackages(TConstArrayView<FString> PackagePaths, EVirtualizationOp
 	}
 
 	OutResultInfo.TimeTaken = FPlatformTime::Seconds() - StartTime;
+	UE_LOG(LogVirtualization, Display, TEXT("Virtualization process complete"));
 	UE_LOG(LogVirtualization, Verbose, TEXT("Virtualization process took %.3f(s)"), OutResultInfo.TimeTaken);
 }
 
