Add code comment explaining why we need to check the file states in revision control to avoid the check being removed in the name of optimization.

#rnx

diff --git a/Engine/Source/Developer/Virtualization/Private/VirtualizationSourceControlUtilities.cpp b/Engine/Source/Developer/Virtualization/Private/VirtualizationSourceControlUtilities.cpp
--- a/Engine/Source/Developer/Virtualization/Private/VirtualizationSourceControlUtilities.cpp
+++ b/Engine/Source/Developer/Virtualization/Private/VirtualizationSourceControlUtilities.cpp
@@ -166,6 +166,9 @@ bool TryCheckoutFiles(const TArray<FString>& FilesToCheckState, TArray<FText>& O
 		return true;
 	}
 
+	// FCheckOut does not store info on exactly which files are checked out which means we cannot accurately fill in
+	// OutFilesCheckedOut unless we first check on the file states.
+	// TODO: Look into changing FCheckOut to provide the info we need to skip out this step.
 	ECommandResult::Type UpdateResult = SCCProvider.GetState(FilesToCheckState, PathStates, EStateCacheUsage::ForceUpdate);
 	if (UpdateResult != ECommandResult::Type::Succeeded)
 	{
