Add some more profiling scopes to the package virtualization project

#rnx

- Attempting to hunt down a performance issue that is currently only occuring for a single user and cannot be reproduced elsewhere.

diff --git a/Engine/Source/Developer/Virtualization/Private/PackageVirtualizationProcess.cpp b/Engine/Source/Developer/Virtualization/Private/PackageVirtualizationProcess.cpp
--- a/Engine/Source/Developer/Virtualization/Private/PackageVirtualizationProcess.cpp
+++ b/Engine/Source/Developer/Virtualization/Private/PackageVirtualizationProcess.cpp
@@ -355,9 +355,11 @@ void VirtualizePackages(TConstArrayView<FString> PackagePaths, EVirtualizationOp
 	Progress.EnterProgressFrame(1.0f);
 
 	TArray<TPair<FPackagePath, FString>> PackagesToReplace;
-
-	// Any package with an updated trailer needs to be copied and an updated trailer appended
+	
 	{
+		// Any package with an updated trailer needs to be copied and an updated trailer appended
+		TRACE_CPUPROFILER_EVENT_SCOPE(UE::Virtualization::VirtualizePackages::DuplicatePackages);
+
 		UE_LOG(LogVirtualization, Display, TEXT("Creating packages with the virtualized data removed..."));
 		for (FPackageInfo& PackageInfo : Packages)
 		{
@@ -380,7 +382,6 @@ void VirtualizePackages(TConstArrayView<FString> PackagePaths, EVirtualizationOp
 			{
 				return;
 			}
-
 		}
 	}
 
@@ -448,22 +449,26 @@ void VirtualizePackages(TConstArrayView<FString> PackagePaths, EVirtualizationOp
 
 		UE_CLOG(!PackagesToReplace.IsEmpty(), LogVirtualization, Display, TEXT("Replacing old packages with the virtualized version..."));
 
-		// Since we had no errors we can now replace all of the packages that were virtualized data with the virtualized replacement file.
-		for(const TPair<FPackagePath,FString>&  Iterator : PackagesToReplace)
 		{
-			const FString OriginalPackagePath = Iterator.Key.GetLocalFullPath();
-			const FString& NewPackagePath = Iterator.Value;
+			// Since we had no errors we can now replace all of the packages that were virtualized data with the virtualized replacement file.
+			TRACE_CPUPROFILER_EVENT_SCOPE(UE::Virtualization::VirtualizePackages::ReplacePackages);
 
-			if (IFileManager::Get().Move(*OriginalPackagePath, *NewPackagePath))
+			for (const TPair<FPackagePath, FString>& Iterator : PackagesToReplace)
 			{
-				OutResultInfo.VirtualizedPackages.Add(OriginalPackagePath);
-			}
-			else
-			{
-				FText Message = FText::Format(	LOCTEXT("Virtualization_MoveFailed", "Unable to replace the package '{0}' with the virtualized version"),
-												FText::FromString(Iterator.Key.GetDebugName()));
-				OutResultInfo.AddError(MoveTemp(Message));
-				continue;
+				const FString OriginalPackagePath = Iterator.Key.GetLocalFullPath();
+				const FString& NewPackagePath = Iterator.Value;
+
+				if (IFileManager::Get().Move(*OriginalPackagePath, *NewPackagePath))
+				{
+					OutResultInfo.VirtualizedPackages.Add(OriginalPackagePath);
+				}
+				else
+				{
+					FText Message = FText::Format(LOCTEXT("Virtualization_MoveFailed", "Unable to replace the package '{0}' with the virtualized version"),
+						FText::FromString(Iterator.Key.GetDebugName()));
+					OutResultInfo.AddError(MoveTemp(Message));
+					continue;
+				}
 			}
 		}
 	}
