Use a specific subdirectory under 'Saved' when duplicating packages for virtualization/hydration so that it is easier to identify disk bloat caused by the duplication and to clean up failed processes.

#rnx

diff --git a/Engine/Source/Developer/Virtualization/Private/PackageUtils.cpp b/Engine/Source/Developer/Virtualization/Private/PackageUtils.cpp
--- a/Engine/Source/Developer/Virtualization/Private/PackageUtils.cpp
+++ b/Engine/Source/Developer/Virtualization/Private/PackageUtils.cpp
@@ -102,8 +102,9 @@ bool TryCopyPackageWithoutTrailer(const FString& SourcePath, const FString& DstP
 
 FString DuplicatePackageWithUpdatedTrailer(const FString& AbsolutePackagePath, const FPackageTrailer& Trailer, TArray<FText>& OutErrors)
 {
+	const FString BaseDir = FPaths::ProjectSavedDir() / TEXT("VADuplication");
 	const FString BaseName = FPaths::GetBaseFilename(AbsolutePackagePath);
-	const FString TempFilePath = FPaths::CreateTempFilename(*FPaths::ProjectSavedDir(), *BaseName.Left(32));
+	const FString TempFilePath = FPaths::CreateTempFilename(*BaseDir, *BaseName.Left(32));
 
 	// TODO Optimization: Combine TryCopyPackageWithoutTrailer with the appending of the new trailer to avoid opening multiple handles
 
@@ -149,8 +150,9 @@ FString DuplicatePackageWithUpdatedTrailer(const FString& AbsolutePackagePath, c
 
 FString DuplicatePackageWithNewTrailer(const FString& AbsolutePackagePath, const FPackageTrailer& OriginalTrailer, const FPackageTrailerBuilder& Builder, TArray<FText>& OutErrors)
 {
+	const FString BaseDir = FPaths::ProjectSavedDir() / TEXT("VADuplication");
 	const FString BaseName = FPaths::GetBaseFilename(AbsolutePackagePath);
-	const FString TempFilePath = FPaths::CreateTempFilename(*FPaths::ProjectSavedDir(), *BaseName.Left(32));
+	const FString TempFilePath = FPaths::CreateTempFilename(*BaseDir, *BaseName.Left(32));
 
 	// TODO Optimization: Combine TryCopyPackageWithoutTrailer with the appending of the new trailer to avoid opening multiple handles
 
