Change the VA filesystem backend to log the absolute file path of its directory so that it is easier for users to find.

- When the path being logged is relative, there is sometimes confusion as to which directory it is relative to.

diff --git a/Engine/Source/Developer/Virtualization/Private/VirtualizationFileBackend.cpp b/Engine/Source/Developer/Virtualization/Private/VirtualizationFileBackend.cpp
--- a/Engine/Source/Developer/Virtualization/Private/VirtualizationFileBackend.cpp
+++ b/Engine/Source/Developer/Virtualization/Private/VirtualizationFileBackend.cpp
@@ -1,3 +1,4 @@
+
 // Copyright Epic Games, Inc. All Rights Reserved.
 
 #include "VirtualizationFileBackend.h"
@@ -57,7 +58,7 @@ bool FFileSystemBackend::Initialize(const FString& ConfigEntry)
 	}
 
 	// Now log a summary of the backend settings to make issues easier to diagnose
-	UE_LOG(LogVirtualization, Log, TEXT("[%s] Using path: '%s'"), *GetDebugName(), *RootDirectory);
+	UE_LOG(LogVirtualization, Log, TEXT("[%s] Using path: '%s'"), *GetDebugName(), *FPaths::ConvertRelativePathToFull(RootDirectory));
 	UE_LOG(LogVirtualization, Log, TEXT("[%s] Will retry failed read attempts %d times with a gap of %dms betwen them"), *GetDebugName(), RetryCount, RetryWaitTimeMS);
 
 	return true;
