Reduce the disk space overhead when virtualizing packages.

- Previously we were duplicating all packages being virtualized, without the recently virtualized payloads, then try to move them one at a time to replace the original packages.
- Now we do all of the prep work such as checking the packages out of revision control and checking that they can be written to first, then duplicate & move each package one at a time.
-- This way we only ever have a single additional duplicated package at a time.
- Moved the final summary logging to an exit scope so we can return when ever we want in the function.
- Added early out for error case.
- Added early out if there are no packages to remove virtualized payloads from.
- Improved the code commands in IVirtualizationSystem::TryVirtualizePackages so that the caller knows the potential failure case scenarios.

diff --git a/Engine/Source/Developer/Virtualization/Private/PackageVirtualizationProcess.cpp b/Engine/Source/Developer/Virtualization/Private/PackageVirtualizationProcess.cpp
--- a/Engine/Source/Developer/Virtualization/Private/PackageVirtualizationProcess.cpp
+++ b/Engine/Source/Developer/Virtualization/Private/PackageVirtualizationProcess.cpp
@@ -9,6 +9,7 @@
 #include "Internationalization/Internationalization.h"
 #include "Misc/PackageName.h"
 #include "Misc/Paths.h"
+#include "Misc/ScopeExit.h"
 #include "Misc/ScopedSlowTask.h"
 #include "PackageUtils.h"
 #include "Serialization/EditorBulkData.h"
@@ -163,6 +164,13 @@ void VirtualizePackages(TConstArrayView<FString> PackagePaths, EVirtualizationOp
 	IVirtualizationSystem& System = IVirtualizationSystem::Get();
 
 	const double StartTime = FPlatformTime::Seconds();
+	
+	ON_SCOPE_EXIT
+	{
+		OutResultInfo.TimeTaken = FPlatformTime::Seconds() - StartTime;
+		UE_LOG(LogVirtualization, Display, TEXT("Virtualization process complete"));
+		UE_LOG(LogVirtualization, Verbose, TEXT("Virtualization process took %.3f(s)"), OutResultInfo.TimeTaken);
+	};
 
 	FScopedSlowTask Progress(5.0f, LOCTEXT("Virtualization_Task", "Virtualizing Assets..."));
 	// Force the task to be visible otherwise it might not be shown if the initial progress frames are too fast
@@ -352,86 +360,98 @@ void VirtualizePackages(TConstArrayView<FString> PackagePaths, EVirtualizationOp
 		}
 	}
 
-	Progress.EnterProgressFrame(1.0f);
-
-	TArray<TPair<FPackagePath, FString>> PackagesToReplace;
-	
+	// If the process has created new errors by this point we should early out before we actually start making changes on disk
+	if (NumErrors != OutResultInfo.GetNumErrors())
 	{
-		// Any package with an updated trailer needs to be copied and an updated trailer appended
-		TRACE_CPUPROFILER_EVENT_SCOPE(UE::Virtualization::VirtualizePackages::DuplicatePackages);
+		return;
+	}
 
-		UE_LOG(LogVirtualization, Display, TEXT("Creating packages with the virtualized data removed..."));
-		for (FPackageInfo& PackageInfo : Packages)
-		{
-			if (!PackageInfo.bWasTrailerUpdated)
-			{
-				continue;
-			}
+	Progress.EnterProgressFrame(1.0f);
 
-			const FPackagePath& PackagePath = PackageInfo.Path; // No need to validate path, we checked this earlier
+	struct FPackageReplacement
+	{
+		FPackagePath Path;
+		FPackageTrailer Trailer;
+	};
 
-			FString NewPackagePath = DuplicatePackageWithUpdatedTrailer(PackagePath.GetLocalFullPath(), PackageInfo.Trailer, OutResultInfo.Errors);
+	TArray<FPackageReplacement> PackagesToReplace;
+	PackagesToReplace.Reserve(Packages.Num());
 
-			if (!NewPackagePath.IsEmpty())
-			{
-				// Now that we have successfully created a new version of the package with an updated trailer 
-				// we need to mark that it should replace the original package.
-				PackagesToReplace.Emplace(PackagePath, MoveTemp(NewPackagePath));
-			}
-			else
-			{
-				return;
-			}
+	for (FPackageInfo& PackageInfo : Packages)
+	{
+		if (PackageInfo.bWasTrailerUpdated)
+		{
+			PackagesToReplace.Add({MoveTemp(PackageInfo.Path), MoveTemp(PackageInfo.Trailer)});
 		}
 	}
+	Packages.Empty(); // No longer used, the useful data have been moved to PackagesToReplace
+
+	if (PackagesToReplace.IsEmpty())
+	{
+		UE_LOG(LogVirtualization, Display, TEXT("No packages need to be updated on disk"));
+		return;
+	}
 
 	UE_LOG(LogVirtualization, Display, TEXT("%d package(s) had their trailer container modified and need to be updated"), PackagesToReplace.Num());
 
-	if (NumErrors == OutResultInfo.GetNumErrors())
 	{
-		// TODO: Consider using the SavePackage model (move the original, then replace, so we can restore all of the original packages if needed)
-		// having said that, once a package is in PackagesToReplace it should still be safe to submit so maybe we don't need this level of protection?
+		// We need to reset the loader of any loaded package that we want save over on disk so that
+		// the file lock is relinquished
+
+		TRACE_CPUPROFILER_EVENT_SCOPE(UE::Virtualization::VirtualizePackages::ResetLoaders);
 
-		UE_CLOG(!PackagesToReplace.IsEmpty(), LogVirtualization, Display, TEXT("Detaching loaded packages from disk..."));
+		UE_LOG(LogVirtualization, Display, TEXT("Detaching loaded packages from disk..."));
 
-		// We need to reset the loader of any package that we want to re-save over
-		for (const TPair<FPackagePath, FString>& Pair : PackagesToReplace)
+		int32 NumPackagesReset = 0;
+		for (const FPackageReplacement& PackageInfo : PackagesToReplace)
 		{
-			UPackage* Package = FindObjectFast<UPackage>(nullptr, Pair.Key.GetPackageFName());
-			if (Package != nullptr)
+			UPackage* LoadedPackage = FindObjectFast<UPackage>(nullptr, PackageInfo.Path.GetPackageFName());
+			if (LoadedPackage != nullptr)
 			{
-				UE_LOG(LogVirtualization, Verbose, TEXT("Detaching '%s'"), *Pair.Key.GetDebugName());
-				ResetLoadersForSave(Package, *Pair.Key.GetLocalFullPath());
+				UE_LOG(LogVirtualization, Verbose, TEXT("Detaching '%s'"), *PackageInfo.Path.GetDebugName());
+				// TODO: Consider using the batch API
+				ResetLoadersForSave(LoadedPackage, *PackageInfo.Path.GetLocalFullPath());
+				NumPackagesReset++;
 			}
 		}
 
-		// Should we try to check out packages from revision control?
-		if (EnumHasAnyFlags(Options, EVirtualizationOptions::Checkout))
-		{
-			UE_CLOG(!PackagesToReplace.IsEmpty(), LogVirtualization, Display, TEXT("Checking out packages from revision control..."));
+		UE_LOG(LogVirtualization, Display, TEXT("Reset the loaders of %d package(s)"), NumPackagesReset);
+	}
 
-			TArray<FString> FilesToCheckState;
-			FilesToCheckState.Reserve(PackagesToReplace.Num());
+	if (EnumHasAnyFlags(Options, EVirtualizationOptions::Checkout))
+	{
+		TRACE_CPUPROFILER_EVENT_SCOPE(UE::Virtualization::VirtualizePackages::Checkout);
 
-			for (const TPair<FPackagePath, FString>& Pair : PackagesToReplace)
-			{
-				FilesToCheckState.Add(Pair.Key.GetLocalFullPath());
-			}
+		UE_LOG(LogVirtualization, Display, TEXT("Checking out packages from revision control..."));
 
-			if (!TryCheckoutFiles(FilesToCheckState, OutResultInfo.Errors, &OutResultInfo.CheckedOutPackages))
-			{
-				return;
-			}
+		TArray<FString> FilesToCheckState;
+		FilesToCheckState.Reserve(PackagesToReplace.Num());
+
+		for (const FPackageReplacement& PackageInfo : PackagesToReplace)
+		{
+			FilesToCheckState.Add(PackageInfo.Path.GetLocalFullPath());
 		}
 
+		if (!TryCheckoutFiles(FilesToCheckState, OutResultInfo.Errors, &OutResultInfo.CheckedOutPackages))
+		{
+			return;
+		}
+
+		UE_LOG(LogVirtualization, Display, TEXT("Checked out %d package(s)"), OutResultInfo.CheckedOutPackages.Num());
+	}
+
+	{
+		TRACE_CPUPROFILER_EVENT_SCOPE(UE::Virtualization::VirtualizePackages::EnsureWritePermissions);
+
 		UE_CLOG(!PackagesToReplace.IsEmpty(), LogVirtualization, Display, TEXT("Checking packages for write access permission..."));
 
 		// Now check to see if there are package files that cannot be edited because they are read only
+		int32 NumSkipped = 0;
 		for (int32 Index = 0; Index < PackagesToReplace.Num(); ++Index)
 		{
-			const TPair<FPackagePath, FString>& Pair = PackagesToReplace[Index];
+			const FPackageReplacement& PackageInfo = PackagesToReplace[Index];
 
-			if (!CanWriteToFile(Pair.Key.GetLocalFullPath()))
+			if (!CanWriteToFile(PackageInfo.Path.GetLocalFullPath()))
 			{
 				// Technically the package could have local payloads that won't be virtualized due to filtering or min payload sizes and so the
 				// following warning is misleading. This will be solved if we move that evaluation to the point of saving a package.
@@ -440,42 +460,52 @@ void VirtualizePackages(TConstArrayView<FString> PackagePaths, EVirtualizationOp
 				// Long term, the stand alone tool should be able to request the UnrealEditor relinquish the lock on the package file so this becomes 
 				// less of a problem.
 				FText Message = FText::Format(LOCTEXT("Virtualization_PkgLocked", "The package file '{0}' has local payloads but is locked for modification and cannot be virtualized, this package will be skipped!"),
-					FText::FromString(Pair.Key.GetDebugName()));
+					FText::FromString(PackageInfo.Path.GetDebugName()));
+
 				UE_LOG(LogVirtualization, Warning, TEXT("%s"), *Message.ToString());
-				
+
 				PackagesToReplace.RemoveAt(Index--);
+				NumSkipped++;
 			}
 		}
 
-		UE_CLOG(!PackagesToReplace.IsEmpty(), LogVirtualization, Display, TEXT("Replacing old packages with the virtualized version..."));
+		UE_CLOG(NumSkipped > 0, LogVirtualization, Warning, TEXT("Skipped %d package(s)"), NumSkipped);
+		UE_CLOG(NumSkipped == 0, LogVirtualization, Display, TEXT("All packages have write permission"));
+	}
+	
+	{
+		// Since we had no errors we can now replace all of the packages that were virtualized data with the virtualized replacement file.
+		TRACE_CPUPROFILER_EVENT_SCOPE(UE::Virtualization::VirtualizePackages::ReplacePackages);
+
+		UE_LOG(LogVirtualization, Display, TEXT("Replacing old packages with the virtualized version..."));
 
+		for (const FPackageReplacement& PackageInfo : PackagesToReplace)
 		{
-			// Since we had no errors we can now replace all of the packages that were virtualized data with the virtualized replacement file.
-			TRACE_CPUPROFILER_EVENT_SCOPE(UE::Virtualization::VirtualizePackages::ReplacePackages);
+			const FString OriginalPackagePath = PackageInfo.Path.GetLocalFullPath();
+			const FString NewPackagePath = DuplicatePackageWithUpdatedTrailer(OriginalPackagePath, PackageInfo.Trailer, OutResultInfo.Errors);
 
-			for (const TPair<FPackagePath, FString>& Iterator : PackagesToReplace)
+			if (NewPackagePath.IsEmpty())
 			{
-				const FString OriginalPackagePath = Iterator.Key.GetLocalFullPath();
-				const FString& NewPackagePath = Iterator.Value;
+				// Duplication failed so skip this package for now. The error will be in OutResultInfo.Errors
+				continue;
+			}
 
-				if (IFileManager::Get().Move(*OriginalPackagePath, *NewPackagePath))
-				{
-					OutResultInfo.VirtualizedPackages.Add(OriginalPackagePath);
-				}
-				else
-				{
-					FText Message = FText::Format(LOCTEXT("Virtualization_MoveFailed", "Unable to replace the package '{0}' with the virtualized version"),
-						FText::FromString(Iterator.Key.GetDebugName()));
-					OutResultInfo.AddError(MoveTemp(Message));
-					continue;
-				}
+			if (IFileManager::Get().Move(*OriginalPackagePath, *NewPackagePath))
+			{
+				OutResultInfo.VirtualizedPackages.Add(OriginalPackagePath);
+			}
+			else
+			{
+				FText Message = FText::Format(LOCTEXT("Virtualization_MoveFailed", "Unable to replace the package '{0}' with the virtualized version"),
+					FText::FromString(PackageInfo.Path.GetDebugName()));
+
+				OutResultInfo.AddError(MoveTemp(Message));
+				continue;
 			}
 		}
-	}
 
-	OutResultInfo.TimeTaken = FPlatformTime::Seconds() - StartTime;
-	UE_LOG(LogVirtualization, Display, TEXT("Virtualization process complete"));
-	UE_LOG(LogVirtualization, Verbose, TEXT("Virtualization process took %.3f(s)"), OutResultInfo.TimeTaken);
+		UE_LOG(LogVirtualization, Display, TEXT("Replaced %d package(s)"), OutResultInfo.VirtualizedPackages.Num());
+	}
 }
 
 } // namespace UE::Virtualization
diff --git a/Engine/Source/Runtime/Core/Public/Virtualization/VirtualizationSystem.h b/Engine/Source/Runtime/Core/Public/Virtualization/VirtualizationSystem.h
--- a/Engine/Source/Runtime/Core/Public/Virtualization/VirtualizationSystem.h
+++ b/Engine/Source/Runtime/Core/Public/Virtualization/VirtualizationSystem.h
@@ -720,6 +720,11 @@ public:
 	/**
 	 * Runs the virtualization process on a set of packages. All of the packages will be parsed and any found to be containing locally stored
 	 * payloads will have them removed but before they are removed they will be pushed to persistent storage.
+	 * 
+	 * Note that if errors occur some or all of the payloads could still get pushed to persistent storage, usually if the errors occur when
+	 * trying to remove the now virtualized payload from the packages on disk.
+	 * In addition, if errors do occur when removing the virtualized payloads, some packages may be virtualized successfully. In any case the
+	 * packages should still be usable and safe to checkin after the process has run, even with failure cases.
 	 *
 	 * @param PackagePaths			An array of file paths to packages that should be virtualized. If a path resolves to a file that is not
 	 *								a valid package then it will be silently skipped and will not be considered an error.
