Log the data reduction when virtualizing packages

- We now log the total size saving by running the virtualization process as well as the original size of the packages being virtualized and their resulting size.
- If LogVirtualization is set to verbose logging we will also log each individual package and its size reduction.

diff --git a/Engine/Source/Developer/Virtualization/Private/PackageVirtualizationProcess.cpp b/Engine/Source/Developer/Virtualization/Private/PackageVirtualizationProcess.cpp
--- a/Engine/Source/Developer/Virtualization/Private/PackageVirtualizationProcess.cpp
+++ b/Engine/Source/Developer/Virtualization/Private/PackageVirtualizationProcess.cpp
@@ -479,6 +479,9 @@ void VirtualizePackages(TConstArrayView<FString> PackagePaths, EVirtualizationOp
 
 		UE_LOG(LogVirtualization, Display, TEXT("Replacing old packages with the virtualized version..."));
 
+		int64 OriginalSizeTotal = 0;
+		int64 ReducedSizeTotal = 0;
+
 		for (const FPackageReplacement& PackageInfo : PackagesToReplace)
 		{
 			const FString OriginalPackagePath = PackageInfo.Path.GetLocalFullPath();
@@ -490,9 +493,20 @@ void VirtualizePackages(TConstArrayView<FString> PackagePaths, EVirtualizationOp
 				continue;
 			}
 
+			const int64 OriginalSize = IFileManager::Get().GetStatData(*OriginalPackagePath).FileSize;
+			const int64 ReducedSize = IFileManager::Get().GetStatData(*NewPackagePath).FileSize;
+
 			if (IFileManager::Get().Move(*OriginalPackagePath, *NewPackagePath))
 			{
 				OutResultInfo.VirtualizedPackages.Add(OriginalPackagePath);
+
+				OriginalSizeTotal += OriginalSize;
+				ReducedSizeTotal += ReducedSize;
+
+				UE_LOG(LogVirtualization, Verbose, TEXT("Reducing %s: %s -> %s"),
+					*PackageInfo.Path.GetDebugName(),
+					*FText::AsMemory(OriginalSize).ToString(),
+					*FText::AsMemory(ReducedSize).ToString());
 			}
 			else
 			{
@@ -504,6 +518,14 @@ void VirtualizePackages(TConstArrayView<FString> PackagePaths, EVirtualizationOp
 			}
 		}
 
+		if (OriginalSizeTotal != ReducedSizeTotal)
+		{
+			UE_LOG(LogVirtualization, Display, TEXT("Total Reduction %s (%s -> %s)"),
+				*FText::AsMemory(OriginalSizeTotal - ReducedSizeTotal).ToString(),
+				*FText::AsMemory(OriginalSizeTotal).ToString(),
+				*FText::AsMemory(ReducedSizeTotal).ToString());
+		}
+
 		UE_LOG(LogVirtualization, Display, TEXT("Replaced %d package(s)"), OutResultInfo.VirtualizedPackages.Num());
 	}
 }
