Clean up the config setup options for the P4SourceControl VA backend.

- Deprecate the config entry 'DepotRoot' and replace it with 'DepotPath'. The use of the term root was not clear and was causing people to think it was the name of the depot rather than the path to the root location to store payloads in.
- If 'DepotPath' is empty or missing, the value for 'ClientStream' will be used for it.
-- The 'DepotPath' can still be set when using streams, if the name of the stream is not a valid path (virtual streams for example.

diff --git a/Engine/Source/Developer/Virtualization/Private/VirtualizationSourceControlBackend.cpp b/Engine/Source/Developer/Virtualization/Private/VirtualizationSourceControlBackend.cpp
--- a/Engine/Source/Developer/Virtualization/Private/VirtualizationSourceControlBackend.cpp
+++ b/Engine/Source/Developer/Virtualization/Private/VirtualizationSourceControlBackend.cpp
@@ -368,7 +368,7 @@ IVirtualizationBackend::EConnectionStatus FSourceControlBackend::OnConnect()
 
 	// When a source control depot is set up a file named 'payload_metainfo.txt' should be submitted to it's root.
 	// This allows us to check for the existence of the file to confirm that the depot root is indeed valid.
-	const FString PayloadMetaInfoPath = WriteToString<512>(DepotRoot, TEXT("payload_metainfo.txt")).ToString();
+	const FString PayloadMetaInfoPath = WriteToString<512>(DepotPathRoot, TEXT("payload_metainfo.txt")).ToString();
 
 	FSharedBuffer MetaInfoBuffer;
 
@@ -624,7 +624,7 @@ bool FSourceControlBackend::PushData(TArrayView<FPushRequest> Requests)
 		if (ClientStream.IsEmpty())
 		{
 			TStringBuilder<512> DepotMapping;
-			DepotMapping << DepotRoot << TEXT("...");
+			DepotMapping << DepotPathRoot << TEXT("...");
 
 			TStringBuilder<128> ClientMapping;
 			ClientMapping << TEXT("//") << WorkspaceName << TEXT("/...");
@@ -835,7 +835,7 @@ bool FSourceControlBackend::DoPayloadsExist(TArrayView<const FIoHash> PayloadIds
 			TStringBuilder<52> LocalPayloadPath;
 			Utils::PayloadIdToPath(PayloadId, LocalPayloadPath);
 
-			DepotPaths.Emplace(WriteToString<512>(DepotRoot, LocalPayloadPath));
+			DepotPaths.Emplace(WriteToString<512>(DepotPathRoot, LocalPayloadPath));
 		}
 	}
 
@@ -864,45 +864,57 @@ bool FSourceControlBackend::DoPayloadsExist(TArrayView<const FIoHash> PayloadIds
 
 bool FSourceControlBackend::TryApplySettingsFromConfigFiles(const FString& ConfigEntry)
 {
-	// We require that a valid depot root has been provided
+	// If the depot root is within a perforce stream then we must specify which stream. This may be a virtual stream with a custom view.
 	{
-		if (!FParse::Value(*ConfigEntry, TEXT("DepotRoot="), DepotRoot))
+		FParse::Value(*ConfigEntry, TEXT("ClientStream="), ClientStream);
+		if (!ClientStream.IsEmpty())
 		{
-			UE_LOG(LogVirtualization, Error, TEXT("'DepotRoot=' not found in the config file"));
-			return false;
+			UE_LOG(LogVirtualization, Log, TEXT("[%s] Using client stream: '%s'"), *GetDebugName(), *ClientStream);
 		}
-
-		if (!DepotRoot.EndsWith(TEXT("/")))
+		else
 		{
-			DepotRoot.AppendChar(TEXT('/'));
+			UE_LOG(LogVirtualization, Log, TEXT("[%s] Not using client stream"), *GetDebugName());
 		}
 	}
 
-	// Now parse the optional config values
-
+	// We require that a valid depot root has been provided
 	{
-		FParse::Value(*ConfigEntry, TEXT("Server="), ServerAddress);
-
-		if (!ServerAddress.IsEmpty())
+		if (FParse::Value(*ConfigEntry, TEXT("DepotRoot="), DepotPathRoot))
 		{
-			UE_LOG(LogVirtualization, Log, TEXT("[%s] Using server address: '%s'"), *GetDebugName(), *ServerAddress);
+			UE_LOG(LogVirtualization, Warning, TEXT("[%s] Entry 'DepotRoot' is deprecated, replace with 'DepotPath'"), *GetDebugName());
 		}
 		else
 		{
-			UE_LOG(LogVirtualization, Log, TEXT("[%s] Will connect to the default server address"), *GetDebugName());
+			FParse::Value(*ConfigEntry, TEXT("DepotPath="), DepotPathRoot);
+		}
+
+		if (DepotPathRoot.IsEmpty())
+		{
+			DepotPathRoot = ClientStream;
+		}
+
+		if (DepotPathRoot.IsEmpty())
+		{
+			UE_LOG(LogVirtualization, Error, TEXT("[%s] 'DepotPath' was not found in the config file!"), *GetDebugName());
+			return false;
+		}
+
+		if (!DepotPathRoot.EndsWith(TEXT("/")))
+		{
+			DepotPathRoot.AppendChar(TEXT('/'));
 		}
 	}
 
-	// If the depot root is within a perforce stream then we must specify which stream. This may be a virtual stream with a custom view.
 	{
-		FParse::Value(*ConfigEntry, TEXT("ClientStream="), ClientStream);
-		if (!ClientStream.IsEmpty())
+		FParse::Value(*ConfigEntry, TEXT("Server="), ServerAddress);
+
+		if (!ServerAddress.IsEmpty())
 		{
-			UE_LOG(LogVirtualization, Log, TEXT("[%s] Using client stream: '%s'"), *GetDebugName(), *ClientStream);
+			UE_LOG(LogVirtualization, Log, TEXT("[%s] Using server address: '%s'"), *GetDebugName(), *ServerAddress);
 		}
 		else
 		{
-			UE_LOG(LogVirtualization, Log, TEXT("[%s] Not using client stream"), *GetDebugName());
+			UE_LOG(LogVirtualization, Log, TEXT("[%s] Will connect to the default server address"), *GetDebugName());
 		}
 	}
 
@@ -996,7 +1008,7 @@ void FSourceControlBackend::CreateDepotPath(const FIoHash& PayloadId, FStringBui
 	TStringBuilder<52> PayloadPath;
 	Utils::PayloadIdToPath(PayloadId, PayloadPath);
 
-	OutPath << DepotRoot << PayloadPath;
+	OutPath << DepotPathRoot << PayloadPath;
 }
 
 bool FSourceControlBackend::FindSubmissionWorkingDir(const FString& ConfigEntry)
diff --git a/Engine/Source/Developer/Virtualization/Private/VirtualizationSourceControlBackend.h b/Engine/Source/Developer/Virtualization/Private/VirtualizationSourceControlBackend.h
--- a/Engine/Source/Developer/Virtualization/Private/VirtualizationSourceControlBackend.h
+++ b/Engine/Source/Developer/Virtualization/Private/VirtualizationSourceControlBackend.h
@@ -22,15 +22,26 @@ class FSemaphore;
  * ----------------------------------------------------------------------------
  * 
  * Ini file setup:
- * 'Name'=(Type=P4SourceControl, DepotRoot="//ABC/")
- * Where 'Name' is the backend name in the hierarchy and 'ABC' is the path in the source control
- * depot where the payload files are being stored.
+ * 'Name'=(Type=P4SourceControl)
+ * Where 'Name' is the backend name in the hierarchy
+ * 
+ * Required Values:
+ * 
+ * DepotPath [string]			The path (in depot format) to the location where virtualized payloads are stored.
+ *								[Default=""]
+ * ClientStream [string]:		Used when the payloads are stored in a stream based depot. It should contain
+ *								the stream name to use when creating a workspace for payload submission. [Default=""]
+ * 
+ * If the payloads are being stored in a depot with type 'stream' them config value 'ClientStream' must be set to a valid
+ * stream. The 'DepotPath' value can still be set to choose a specific location within that stream or when the stream 
+ * name is not a valid depot path (such as virtual streams)
+ * 
+ * If the payloads are not being stored in a 'stream' depot then only 'DepotPath' is required.
  * 
  * Optional Values:
+ * 
  * Server [string]:				When set the backend will use this server address to connect to. When not
  *								set the backend will use the environment defaults. [Default=""]
- * ClientStream [string]:		Used when the payloads are stored in a stream based depot. It should contain
- *								the stream name to use when creating a workspace for payload submission. [Default=""]
  * UsePartitionedClient [bool]:	When true the temporary workspace client created to submit payloads 
  *								from will be created as a partitioned workspace which is less overhead
  *								on the source control server. If your server does not support this then
@@ -140,8 +151,8 @@ private:
 	/** The name of the current project */
 	FString ProjectName;
 
-	/** The root where the virtualized payloads are stored in source control */
-	FString DepotRoot;
+	/** The path (in depot format) to where the virtualized payloads are stored in source control */
+	FString DepotPathRoot;
 
 	/** Address of the server to connect to */
 	FString ServerAddress;
