Add bitflag enums to IVirtualizationBackend::PullData and IVirtualizationBackend::PushData so that it is easier to specialize behavior in the future.

- The enums don't contain any entries at the moment, the goal of this change is to reduce the friction that future API changes might have.
- The older versions of PullData/PushData have been deprecated but this won't flag compiler warnings if people have derived their own implementations of IVirtualizationBackend, so in addition to deprecating the methods have been set to 'final' and have checkNoEntry implementations which is the best we can do at the moment.
-- It is unlikely that people have actually done this as the documentation on adding new backend types will only go out in 5.3.

diff --git a/Engine/Source/Developer/Virtualization/Private/VirtualizationDDCBackend.cpp b/Engine/Source/Developer/Virtualization/Private/VirtualizationDDCBackend.cpp
--- a/Engine/Source/Developer/Virtualization/Private/VirtualizationDDCBackend.cpp
+++ b/Engine/Source/Developer/Virtualization/Private/VirtualizationDDCBackend.cpp
@@ -74,7 +74,7 @@ IVirtualizationBackend::EConnectionStatus FDDCBackend::OnConnect()
 	return IVirtualizationBackend::EConnectionStatus::Connected;
 }
 
-bool FDDCBackend::PushData(TArrayView<FPushRequest> Requests)
+bool FDDCBackend::PushData(TArrayView<FPushRequest> Requests, EPushFlags Flags)
 {
 	TRACE_CPUPROFILER_EVENT_SCOPE(FDDCBackend::PushData);
 
@@ -131,7 +131,7 @@ bool FDDCBackend::PushData(TArrayView<FPushRequest> Requests)
 	return bWasSuccess;
 }
 
-bool FDDCBackend::PullData(TArrayView<FPullRequest> Requests)
+bool FDDCBackend::PullData(TArrayView<FPullRequest> Requests, EPullFlags Flags)
 {
 	TRACE_CPUPROFILER_EVENT_SCOPE(FDDCBackend::PullData);
 
diff --git a/Engine/Source/Developer/Virtualization/Private/VirtualizationDDCBackend.h b/Engine/Source/Developer/Virtualization/Private/VirtualizationDDCBackend.h
--- a/Engine/Source/Developer/Virtualization/Private/VirtualizationDDCBackend.h
+++ b/Engine/Source/Developer/Virtualization/Private/VirtualizationDDCBackend.h
@@ -38,8 +38,8 @@ private:
 
 	virtual EConnectionStatus OnConnect() override;
 
-	virtual bool PushData(TArrayView<FPushRequest> Requests) override;
-	virtual bool PullData(TArrayView<FPullRequest> Requests) override;
+	virtual bool PushData(TArrayView<FPushRequest> Requests, EPushFlags Flags) override;
+	virtual bool PullData(TArrayView<FPullRequest> Requests, EPullFlags Flags) override;
 	
 	virtual bool DoesPayloadExist(const FIoHash& Id) override;
 
diff --git a/Engine/Source/Developer/Virtualization/Private/VirtualizationFileBackend.cpp b/Engine/Source/Developer/Virtualization/Private/VirtualizationFileBackend.cpp
--- a/Engine/Source/Developer/Virtualization/Private/VirtualizationFileBackend.cpp
+++ b/Engine/Source/Developer/Virtualization/Private/VirtualizationFileBackend.cpp
@@ -68,7 +68,7 @@ IVirtualizationBackend::EConnectionStatus FFileSystemBackend::OnConnect()
 	return IVirtualizationBackend::EConnectionStatus::Connected;
 }
 
-bool FFileSystemBackend::PushData(TArrayView<FPushRequest> Requests)
+bool FFileSystemBackend::PushData(TArrayView<FPushRequest> Requests, EPushFlags Flags)
 {
 	TRACE_CPUPROFILER_EVENT_SCOPE(FFileSystemBackend::PushData);
 
@@ -177,7 +177,7 @@ bool FFileSystemBackend::PushData(TArrayView<FPushRequest> Requests)
 	return ErrorCount == 0;
 }
 
-bool FFileSystemBackend::PullData(TArrayView<FPullRequest> Requests)
+bool FFileSystemBackend::PullData(TArrayView<FPullRequest> Requests, EPullFlags Flags)
 {
 	TRACE_CPUPROFILER_EVENT_SCOPE(FFileSystemBackend::PullData);
 
diff --git a/Engine/Source/Developer/Virtualization/Private/VirtualizationFileBackend.h b/Engine/Source/Developer/Virtualization/Private/VirtualizationFileBackend.h
--- a/Engine/Source/Developer/Virtualization/Private/VirtualizationFileBackend.h
+++ b/Engine/Source/Developer/Virtualization/Private/VirtualizationFileBackend.h
@@ -44,8 +44,8 @@ private:
 
 	virtual EConnectionStatus OnConnect() override;
 
-	virtual bool PushData(TArrayView<FPushRequest> Requests) override;
-	virtual bool PullData(TArrayView<FPullRequest> Requests) override;
+	virtual bool PushData(TArrayView<FPushRequest> Requests, EPushFlags Flags) override;
+	virtual bool PullData(TArrayView<FPullRequest> Requests, EPullFlags Flags) override;
 
 	virtual bool DoesPayloadExist(const FIoHash& Id) override;
 	
diff --git a/Engine/Source/Developer/Virtualization/Private/VirtualizationManager.cpp b/Engine/Source/Developer/Virtualization/Private/VirtualizationManager.cpp
--- a/Engine/Source/Developer/Virtualization/Private/VirtualizationManager.cpp
+++ b/Engine/Source/Developer/Virtualization/Private/VirtualizationManager.cpp
@@ -1776,7 +1776,7 @@ bool FVirtualizationManager::TryCacheDataToBackend(IVirtualizationBackend& Backe
 	COOK_STAT(FCookStats::FScopedStatsCounter Timer(Stats));
 	COOK_STAT(Timer.TrackCyclesOnly());
 	
-	if (Backend.PushData(Requests))
+	if (Backend.PushData(Requests, IVirtualizationBackend::EPushFlags::None))
 	{
 #if ENABLE_COOK_STATS
 		Timer.AddHit(0);
@@ -1807,7 +1807,7 @@ bool FVirtualizationManager::TryPushDataToBackend(IVirtualizationBackend& Backen
 	COOK_STAT(FCookStats::FScopedStatsCounter Timer(Stats));
 	COOK_STAT(Timer.TrackCyclesOnly());
 
-	const bool bPushResult = Backend.PushData(Requests);
+	const bool bPushResult = Backend.PushData(Requests, IVirtualizationBackend::EPushFlags::None);
 
 #if ENABLE_COOK_STATS
 	if (bPushResult)
@@ -1892,7 +1892,7 @@ void FVirtualizationManager::PullDataFromBackend(IVirtualizationBackend& Backend
 	COOK_STAT(FCookStats::FScopedStatsCounter Timer(Stats));
 	COOK_STAT(Timer.TrackCyclesOnly());
 	
-	Backend.PullData(Requests);
+	Backend.PullData(Requests, IVirtualizationBackend::EPullFlags::None);
 	
 #if ENABLE_COOK_STATS
 	const bool bIsInGameThread = IsInGameThread();
diff --git a/Engine/Source/Developer/Virtualization/Private/VirtualizationSourceControlBackend.cpp b/Engine/Source/Developer/Virtualization/Private/VirtualizationSourceControlBackend.cpp
--- a/Engine/Source/Developer/Virtualization/Private/VirtualizationSourceControlBackend.cpp
+++ b/Engine/Source/Developer/Virtualization/Private/VirtualizationSourceControlBackend.cpp
@@ -403,7 +403,7 @@ IVirtualizationBackend::EConnectionStatus FSourceControlBackend::OnConnect()
 	return IVirtualizationBackend::EConnectionStatus::Connected;
 }
 
-bool FSourceControlBackend::PullData(TArrayView<FPullRequest> Requests)
+bool FSourceControlBackend::PullData(TArrayView<FPullRequest> Requests, EPullFlags Flags)
 {
 	TRACE_CPUPROFILER_EVENT_SCOPE(FSourceControlBackend::PullData);
 
@@ -484,7 +484,7 @@ bool FSourceControlBackend::DoesPayloadExist(const FIoHash& Id)
 	}
 }
 
-bool FSourceControlBackend::PushData(TArrayView<FPushRequest> Requests)
+bool FSourceControlBackend::PushData(TArrayView<FPushRequest> Requests, EPushFlags Flags)
 {
 	TRACE_CPUPROFILER_EVENT_SCOPE(FSourceControlBackend::PushData);
 
diff --git a/Engine/Source/Developer/Virtualization/Private/VirtualizationSourceControlBackend.h b/Engine/Source/Developer/Virtualization/Private/VirtualizationSourceControlBackend.h
--- a/Engine/Source/Developer/Virtualization/Private/VirtualizationSourceControlBackend.h
+++ b/Engine/Source/Developer/Virtualization/Private/VirtualizationSourceControlBackend.h
@@ -127,8 +127,8 @@ private:
 
 	virtual EConnectionStatus OnConnect() override;
 	
-	virtual bool PushData(TArrayView<FPushRequest> Requests) override;
-	virtual bool PullData(TArrayView<FPullRequest> Requests) override;
+	virtual bool PushData(TArrayView<FPushRequest> Requests, EPushFlags Flags) override;
+	virtual bool PullData(TArrayView<FPullRequest> Requests, EPullFlags Flags) override;
 
 	virtual bool DoesPayloadExist(const FIoHash& Id) override;
 	
diff --git a/Engine/Source/Developer/Virtualization/Public/IVirtualizationBackend.h b/Engine/Source/Developer/Virtualization/Public/IVirtualizationBackend.h
--- a/Engine/Source/Developer/Virtualization/Public/IVirtualizationBackend.h
+++ b/Engine/Source/Developer/Virtualization/Public/IVirtualizationBackend.h
@@ -13,7 +13,6 @@ struct FIoHash;
 
 namespace UE::Virtualization
 {
-
 /**
  * The interface to derive from to create a new backend implementation.
  * 
@@ -51,6 +50,22 @@ public:
 		Connected,
 	};
 
+	/** Specialize pushing behavior, @See IVirtualizationBackend::PushData */
+	enum class EPushFlags
+	{
+		None = 0
+	};
+
+	FRIEND_ENUM_CLASS_FLAGS(EPushFlags);
+
+	/** Specialize pulling behavior, @See IVirtualizationBackend::PullData */
+	enum class EPullFlags
+	{
+		None = 0
+	};
+
+	FRIEND_ENUM_CLASS_FLAGS(EPullFlags);
+
 protected:
 
 	IVirtualizationBackend(FStringView InConfigName, FStringView InDebugName, EOperations InSupportedOperations)
@@ -103,7 +118,7 @@ public:
 	/**
 	 * @return True if all of the requests succeeded, false if one of more failed
 	 */
-	virtual bool PushData(TArrayView<FPushRequest> Requests) = 0;
+	virtual bool PushData(TArrayView<FPushRequest> Requests, EPushFlags Flags) = 0;
 
 	/** 
 	 * The backend will attempt to retrieve the given payloads by what ever method the backend uses.
@@ -125,7 +140,7 @@ public:
 	// Assume that the array only has unique requests
 	// Will set request error value if there is a problem with loading it
 	// Returns false on critical error, otherwise true
-	virtual bool PullData(TArrayView<FPullRequest> Requests) = 0;
+	virtual bool PullData(TArrayView<FPullRequest> Requests, EPullFlags Flags) = 0;
 	
 	/**
 	 * Checks if a payload exists in the backends storage.
@@ -225,6 +240,25 @@ private:
 	/** Override to implement the backends connection code */
 	virtual EConnectionStatus OnConnect() = 0;
 
+private:
+	// Deprecated methods, will cause compiler errors if anyone has actually overridden IVirtualizationBackend
+	// pre 5.3 but there is currently no good compiler support for invoking compiler warnings when overriding
+	// a deprecated base method.
+
+	UE_DEPRECATED(5.2, "Use the overload that also takes EPushFlags")
+	virtual bool PushData(TArrayView<FPushRequest> Requests) final
+	{
+		checkNoEntry();
+		return false;
+	}
+
+	UE_DEPRECATED(5.3, "Use the overload that also takes EPullFlags")
+	virtual bool PullData(TArrayView<FPullRequest> Requests) final
+	{
+		checkNoEntry();
+		return false;
+	}
+
 private:
 
 	/** The operations that this backend supports */
