Use the FMessageDialog overloads that pass the optional title by-value

#rnx

diff --git a/Engine/Source/Developer/Virtualization/Private/VirtualizationManager.cpp b/Engine/Source/Developer/Virtualization/Private/VirtualizationManager.cpp
--- a/Engine/Source/Developer/Virtualization/Private/VirtualizationManager.cpp
+++ b/Engine/Source/Developer/Virtualization/Private/VirtualizationManager.cpp
@@ -1939,7 +1939,7 @@ FVirtualizationManager::ErrorHandlingResult FVirtualizationManager::OnPayloadPul
 
 		const FText Message = MsgBuilder.ToText();
 
-		EAppReturnType::Type Result = FMessageDialog::Open(EAppMsgType::YesNo, EAppReturnType::No, Message, &Title);
+		EAppReturnType::Type Result = FMessageDialog::Open(EAppMsgType::YesNo, EAppReturnType::No, Message, Title);
 
 		if (Result == EAppReturnType::No)
 		{
