Added code documentation for 'PullErrorAdditionalMsg' that was missing from CL 24714178

#rnx

diff --git a/Engine/Source/Developer/Virtualization/Private/VirtualizationManager.h b/Engine/Source/Developer/Virtualization/Private/VirtualizationManager.h
--- a/Engine/Source/Developer/Virtualization/Private/VirtualizationManager.h
+++ b/Engine/Source/Developer/Virtualization/Private/VirtualizationManager.h
@@ -119,6 +119,10 @@ struct FAnalyticsEventAttribute;
  *												pull will return an error and allow the process to carry on (the original error handling logic)
  *												and when false a dialog will be displayed to the user warning them about the failed pull and 
  *												prompting them to retry the pull or to quit the process. [Default=true]
+ * PullErrorAdditionalMsg [string]				An additional message that will be added to the error dialog presented on payload pull failure.
+ *												This allows you to add custom information, such as links to internal help docs without editing
+ *												code. Note that this additional message only works with the error dialog and will do nothing
+ *												if 'UseLegacyErrorHandling' is true. [Default=""]
  */
 
 namespace UE::Virtualization
