Add a new config option 'ForceCachingOnPull' to VA that when true will force backends to re-upload payloads when caching.

- Most backends support some form of existence check which is run before a push to prevent uploading data that is already there. Sometimes we want to ignore that check (bugs in the backend etc) and need a way to force that to happen.
- The new config option [Core.VirtualizationModule]:ForceCachingOnPull (default false) will do this when set to true.
- The file and DDC backends will respect the flag but for now the source control backend will continue to check to avoid accidental pointless revisions being added.

diff --git a/Engine/Config/BaseEngine.ini b/Engine/Config/BaseEngine.ini
--- a/Engine/Config/BaseEngine.ini
+++ b/Engine/Config/BaseEngine.ini
@@ -1400,6 +1400,7 @@ FilterMode=OptOut
 FilterMapContent=True
 LazyInitConnections=False
 UseLegacyErrorHandling=True
+ForceCachingOnPull=False
 PullErrorAdditionalMsg=""
 ; The following assets should not virtualize until it is confirmed that they do not crash when virtualized payloads are missing
 +DisabledAsset=StaticMesh
diff --git a/Engine/Source/Developer/Virtualization/Private/VirtualizationDDCBackend.cpp b/Engine/Source/Developer/Virtualization/Private/VirtualizationDDCBackend.cpp
--- a/Engine/Source/Developer/Virtualization/Private/VirtualizationDDCBackend.cpp
+++ b/Engine/Source/Developer/Virtualization/Private/VirtualizationDDCBackend.cpp
@@ -84,12 +84,14 @@ bool FDDCBackend::PushData(TArrayView<FPushRequest> Requests, EPushFlags Flags)
 
 	bool bWasSuccess = true;
 
+	const bool bEnableExistenceCheck = !EnumHasAllFlags(Flags, EPushFlags::Force);
+
 	// TODO: We tend not to memory bloat too much on large batches as the requests complete quite quickly
 	// however we might want to consider adding better control on how much total memory we can dedicate to
 	// loading payloads before we wait for requests to complete?
 	for (FPushRequest& Request : Requests)
 	{
-		if (DoesPayloadExist(Request.GetIdentifier()))
+		if (bEnableExistenceCheck && DoesPayloadExist(Request.GetIdentifier()))
 		{
 			Request.SetResult(FPushResult::GetAsAlreadyExists());
 		}
@@ -134,7 +136,7 @@ bool FDDCBackend::PushData(TArrayView<FPushRequest> Requests, EPushFlags Flags)
 bool FDDCBackend::PullData(TArrayView<FPullRequest> Requests, EPullFlags Flags)
 {
 	TRACE_CPUPROFILER_EVENT_SCOPE(FDDCBackend::PullData);
-
+#if 0
 	UE::DerivedData::ICache& Cache = UE::DerivedData::GetCache();
 
 	UE::DerivedData::FRequestOwner Owner(UE::DerivedData::EPriority::Normal);
@@ -161,7 +163,7 @@ bool FDDCBackend::PullData(TArrayView<FPullRequest> Requests, EPullFlags Flags)
 	}
 
 	Owner.Wait();
-
+#endif
 	return true;
 }
 
diff --git a/Engine/Source/Developer/Virtualization/Private/VirtualizationFileBackend.cpp b/Engine/Source/Developer/Virtualization/Private/VirtualizationFileBackend.cpp
--- a/Engine/Source/Developer/Virtualization/Private/VirtualizationFileBackend.cpp
+++ b/Engine/Source/Developer/Virtualization/Private/VirtualizationFileBackend.cpp
@@ -74,11 +74,13 @@ bool FFileSystemBackend::PushData(TArrayView<FPushRequest> Requests, EPushFlags
 
 	int32 ErrorCount = 0;
 
+	const bool bEnableExistenceCheck = !EnumHasAllFlags(Flags, EPushFlags::Force);
+
 	for (FPushRequest& Request : Requests)
 	{
 		const FIoHash& PayloadId = Request.GetIdentifier();
-
-		if (DoesPayloadExist(PayloadId))
+	
+		if (bEnableExistenceCheck && DoesPayloadExist(PayloadId))
 		{
 			UE_LOG(LogVirtualization, Verbose, TEXT("[%s] Already has a copy of the payload '%s'."), *GetDebugName(), *LexToString(PayloadId));
 			Request.SetResult(FPushResult::GetAsAlreadyExists());
diff --git a/Engine/Source/Developer/Virtualization/Private/VirtualizationManager.cpp b/Engine/Source/Developer/Virtualization/Private/VirtualizationManager.cpp
--- a/Engine/Source/Developer/Virtualization/Private/VirtualizationManager.cpp
+++ b/Engine/Source/Developer/Virtualization/Private/VirtualizationManager.cpp
@@ -4,7 +4,6 @@
 
 #include "HAL/IConsoleManager.h"
 #include "HAL/PlatformTime.h"
-#include "IVirtualizationBackend.h"
 #include "Logging/MessageLog.h"
 #include "Misc/CommandLine.h"
 #include "Misc/ConfigCacheIni.h"
@@ -493,6 +492,7 @@ FVirtualizationManager::FVirtualizationManager()
 	, bAllowSubmitIfVirtualizationFailed(false)
 	, bLazyInitConnections(false)
 	, bUseLegacyErrorHandling(true)
+	, bForceCachingOnPull(false)
 	, bPendingBackendConnections(false)
 {
 }
@@ -1134,6 +1134,11 @@ void FVirtualizationManager::ApplySettingsFromConfigFiles(const FConfigFile& Con
 		// This value is not echoed to the log file, as seeing an error string there might confuse users
 	}
 
+	{
+		ConfigFile.GetBool(ConfigSection, TEXT("ForceCachingOnPull"), bForceCachingOnPull);
+		UE_LOG(LogVirtualization, Display, TEXT("\tForceCachingOnPull : %s"), bForceCachingOnPull ? TEXT("true") : TEXT("false"));
+	}
+
 	// Deprecated
 	{
 		bool bDummyValue = true;
@@ -1728,7 +1733,7 @@ void FVirtualizationManager::EnsureBackendConnections()
 	}
 }
 
-void FVirtualizationManager::CachePayloads(TArrayView<FPushRequest> Requests, const IVirtualizationBackend* BackendSource)
+void FVirtualizationManager::CachePayloads(TArrayView<FPushRequest> Requests, const IVirtualizationBackend* BackendSource, IVirtualizationBackend::EPushFlags Flags)
 {
 	TRACE_CPUPROFILER_EVENT_SCOPE(FVirtualizationManager::CachePayload);
 
@@ -1743,7 +1748,7 @@ void FVirtualizationManager::CachePayloads(TArrayView<FPushRequest> Requests, co
 			return;
 		}
 
-		const bool bResult = TryCacheDataToBackend(*BackendToCache, Requests);
+		const bool bResult = TryCacheDataToBackend(*BackendToCache, Requests, Flags);
 
 		if (!bResult)
 		{
@@ -1770,13 +1775,13 @@ void FVirtualizationManager::CachePayloads(TArrayView<FPushRequest> Requests, co
 	}
 }
 
-bool FVirtualizationManager::TryCacheDataToBackend(IVirtualizationBackend& Backend, TArrayView<FPushRequest> Requests)
+bool FVirtualizationManager::TryCacheDataToBackend(IVirtualizationBackend& Backend, TArrayView<FPushRequest> Requests, IVirtualizationBackend::EPushFlags Flags)
 {
 	COOK_STAT(FCookStats::CallStats & Stats = Profiling::GetCacheStats(Backend));
 	COOK_STAT(FCookStats::FScopedStatsCounter Timer(Stats));
 	COOK_STAT(Timer.TrackCyclesOnly());
 	
-	if (Backend.PushData(Requests, IVirtualizationBackend::EPushFlags::None))
+	if (Backend.PushData(Requests, Flags))
 	{
 #if ENABLE_COOK_STATS
 		Timer.AddHit(0);
@@ -1865,7 +1870,10 @@ void FVirtualizationManager::PullDataFromAllBackends(TArrayView<FPullRequest> Re
 			TArray<FPushRequest> PayloadsToCache = RequestsCollection.OnPullCompleted(*Backend, bShouldCache);
 			if (!PayloadsToCache.IsEmpty())
 			{
-				CachePayloads(PayloadsToCache, Backend);
+				const IVirtualizationBackend::EPushFlags CacheOnPullFlags = bForceCachingOnPull	? IVirtualizationBackend::EPushFlags::Force
+																								: IVirtualizationBackend::EPushFlags::None;
+
+				CachePayloads(PayloadsToCache, Backend, CacheOnPullFlags);
 			}
 
 			// We can early out if there is no more requests to make
diff --git a/Engine/Source/Developer/Virtualization/Private/VirtualizationManager.h b/Engine/Source/Developer/Virtualization/Private/VirtualizationManager.h
--- a/Engine/Source/Developer/Virtualization/Private/VirtualizationManager.h
+++ b/Engine/Source/Developer/Virtualization/Private/VirtualizationManager.h
@@ -4,10 +4,9 @@
 
 #include "Compression/CompressedBuffer.h"
 #include "HAL/CriticalSection.h"
+#include "IVirtualizationBackend.h"
 #include "Logging/LogMacros.h"
 #include "Templates/UniquePtr.h"
-
-
 #include "Virtualization/VirtualizationSystem.h"
 
 class IConsoleObject;
@@ -123,12 +122,15 @@ struct FAnalyticsEventAttribute;
  *												This allows you to add custom information, such as links to internal help docs without editing
  *												code. Note that this additional message only works with the error dialog and will do nothing
  *												if 'UseLegacyErrorHandling' is true. [Default=""]
+ * ForceCachingOnPull [bool]:					When true backends will be told to always upload the payload when a caching as a result of 
+ *												a payload pull as in this scenario we already know that the backend failed to pull the payload
+ *												before it was pulled from a backend later in the hierarchy. Can be used to try and skip
+ *												expensive existence checks, or if a backend is in a bad state where it believes it has the payload
+ *												but is unable to actually return the data. [Default=false]
  */
 
 namespace UE::Virtualization
 {
-class IVirtualizationBackend;
-class IVirtualizationBackendFactory;
 class FPullRequestCollection;
 
 /** The default mode of filtering to use with package paths that do not match entries in UVirtualizationFilterSettings */
@@ -215,9 +217,9 @@ private:
 
 	void EnsureBackendConnections();
 
-	void CachePayloads(TArrayView<FPushRequest> Requests, const IVirtualizationBackend* BackendSource);
+	void CachePayloads(TArrayView<FPushRequest> Requests, const IVirtualizationBackend* BackendSource, IVirtualizationBackend::EPushFlags Flags);
 
-	bool TryCacheDataToBackend(IVirtualizationBackend& Backend, TArrayView<FPushRequest> Requests);
+	bool TryCacheDataToBackend(IVirtualizationBackend& Backend, TArrayView<FPushRequest> Requests, IVirtualizationBackend::EPushFlags Flags);
 	bool TryPushDataToBackend(IVirtualizationBackend& Backend, TArrayView<FPushRequest> Requests);
 
 	void PullDataFromAllBackends(TArrayView<FPullRequest> Requests);
@@ -311,6 +313,9 @@ private:
 	/** When true we do not display an error dialog on failed payload pulling and rely on the caller handling it */
 	bool bUseLegacyErrorHandling;
 
+	/** When true IVirtualizationBackend::EPushFlags::Force will be passed to backends that need to cache payloads when pulling */
+	bool bForceCachingOnPull;
+
 	/** An additional error message to display when pulling payloads fails */
 	FString PullErrorAdditionalMsg;
 
diff --git a/Engine/Source/Developer/Virtualization/Private/VirtualizationSourceControlBackend.cpp b/Engine/Source/Developer/Virtualization/Private/VirtualizationSourceControlBackend.cpp
--- a/Engine/Source/Developer/Virtualization/Private/VirtualizationSourceControlBackend.cpp
+++ b/Engine/Source/Developer/Virtualization/Private/VirtualizationSourceControlBackend.cpp
@@ -522,6 +522,11 @@ bool FSourceControlBackend::PushData(TArrayView<FPushRequest> Requests, EPushFla
 	TArray<const FPushRequest*> RequestsToPush;
 	RequestsToPush.Reserve(UniqueRequests.Num());
 
+	// Note that we do not check EPushFlags::Force here and always check if the payloads are already in
+	// source control as we do not want multiple revisions.
+	// This might change at some point if we want to try changing the compression codec of the stored
+	// payloads.
+
 	{
 		TRACE_CPUPROFILER_EVENT_SCOPE(FSourceControlBackend::PushData::CheckIfPayloadsAlreadyExist);
 
diff --git a/Engine/Source/Developer/Virtualization/Public/IVirtualizationBackend.h b/Engine/Source/Developer/Virtualization/Public/IVirtualizationBackend.h
--- a/Engine/Source/Developer/Virtualization/Public/IVirtualizationBackend.h
+++ b/Engine/Source/Developer/Virtualization/Public/IVirtualizationBackend.h
@@ -51,15 +51,17 @@ public:
 	};
 
 	/** Specialize pushing behavior, @See IVirtualizationBackend::PushData */
-	enum class EPushFlags
+	enum class EPushFlags : uint32
 	{
-		None = 0
+		None = 0,
+		/** Backends should not test payloads for existence and just assume that they need re-uploading */
+		Force = 1 << 0
 	};
 
 	FRIEND_ENUM_CLASS_FLAGS(EPushFlags);
 
 	/** Specialize pulling behavior, @See IVirtualizationBackend::PullData */
-	enum class EPullFlags
+	enum class EPullFlags : uint32
 	{
 		None = 0
 	};
