Do not record VA caching statistics if no payload was actually cached because they were already present in the backend.

#rnx

- If all payloads were already in the backend we would still record the time taken to query but do not increment the number of hits nor the size of the payloads. Without a stat for "payloads skipped" which we do not currently have any interest in adding, recording the time taken to query where no transfer takes place just ends up being misleading and looks like a bug.

diff --git a/Engine/Source/Developer/Virtualization/Private/VirtualizationManager.cpp b/Engine/Source/Developer/Virtualization/Private/VirtualizationManager.cpp
--- a/Engine/Source/Developer/Virtualization/Private/VirtualizationManager.cpp
+++ b/Engine/Source/Developer/Virtualization/Private/VirtualizationManager.cpp
@@ -1784,9 +1784,8 @@ bool FVirtualizationManager::TryCacheDataToBackend(IVirtualizationBackend& Backe
 	if (Backend.PushData(Requests, Flags))
 	{
 #if ENABLE_COOK_STATS
-		Timer.AddHit(0);
-
 		const bool bIsInGameThread = IsInGameThread();
+		bool bWasDataPushed = false;
 
 		for (const FPushRequest& Request : Requests)
 		{
@@ -1794,8 +1793,15 @@ bool FVirtualizationManager::TryCacheDataToBackend(IVirtualizationBackend& Backe
 			{
 				Stats.Accumulate(FCookStats::CallStats::EHitOrMiss::Hit, FCookStats::CallStats::EStatType::Counter, 1l, bIsInGameThread);
 				Stats.Accumulate(FCookStats::CallStats::EHitOrMiss::Hit, FCookStats::CallStats::EStatType::Bytes, Request.GetPayload().GetCompressedSize(), bIsInGameThread);
+
+				bWasDataPushed = true;
 			}
 		}
+
+		if (bWasDataPushed)
+		{
+			Timer.AddHit(0);
+		}
 #endif // ENABLE_COOK_STATS
 
 		return true;
