Remove debug code accidently submitted in CL 25066347

#rnx

diff --git a/Engine/Source/Developer/Virtualization/Private/VirtualizationDDCBackend.cpp b/Engine/Source/Developer/Virtualization/Private/VirtualizationDDCBackend.cpp
--- a/Engine/Source/Developer/Virtualization/Private/VirtualizationDDCBackend.cpp
+++ b/Engine/Source/Developer/Virtualization/Private/VirtualizationDDCBackend.cpp
@@ -136,7 +136,7 @@ bool FDDCBackend::PushData(TArrayView<FPushRequest> Requests, EPushFlags Flags)
 bool FDDCBackend::PullData(TArrayView<FPullRequest> Requests, EPullFlags Flags)
 {
 	TRACE_CPUPROFILER_EVENT_SCOPE(FDDCBackend::PullData);
-#if 0
+
 	UE::DerivedData::ICache& Cache = UE::DerivedData::GetCache();
 
 	UE::DerivedData::FRequestOwner Owner(UE::DerivedData::EPriority::Normal);
@@ -163,7 +163,7 @@ bool FDDCBackend::PullData(TArrayView<FPullRequest> Requests, EPullFlags Flags)
 	}
 
 	Owner.Wait();
-#endif
+
 	return true;
 }
 
