Remove duplicate package entries from the list of packages to virtualize/rehydrate.

#rnx

- One crash was reported where a user has somehow managed to try to submit a list of packages containing duplicate entries. This means we try to strip the trailer from that package multiple times which would cause the file system to assert during a call to truncate.
- This is technically the fault of higher level systems outside the control of VA but we should still handle it properly.
- Now duplicate entries are detected and removed so that each package is only processed once.
- If duplicates are detected we log a warning.

diff --git a/Engine/Source/Developer/Virtualization/Private/PackageRehydrationProcess.cpp b/Engine/Source/Developer/Virtualization/Private/PackageRehydrationProcess.cpp
--- a/Engine/Source/Developer/Virtualization/Private/PackageRehydrationProcess.cpp
+++ b/Engine/Source/Developer/Virtualization/Private/PackageRehydrationProcess.cpp
@@ -197,11 +197,22 @@ void RehydratePackages(TConstArrayView<FString> PackagePaths, ERehydrationOption
 
 	double Time = FPlatformTime::Seconds();
 
+	TSet<FString> ConsideredPackages;
+	ConsideredPackages.Reserve(PackagePaths.Num());
+
 	// Attempt to rehydrate the packages
 	for(int32 Index = 0; Index < PackagePaths.Num(); ++Index)
 	{
 		const FString& FilePath = PackagePaths[Index];
 
+		bool bIsDuplicate = false;
+		ConsideredPackages.Add(FilePath, &bIsDuplicate);
+		if (bIsDuplicate)
+		{
+			UE_LOG(LogVirtualization, Verbose, TEXT("Skipping duplicate package entry '%s'"), *FilePath);
+			continue; // Skip duplicate packages
+		}
+
 		FPackageTrailer Trailer;
 		FPackageTrailerBuilder Builder;
 
@@ -229,6 +240,11 @@ void RehydratePackages(TConstArrayView<FString> PackagePaths, ERehydrationOption
 		}
 	}
 
+	const int32 NumSkippedPackages = PackagePaths.Num() - ConsideredPackages.Num();
+	ConsideredPackages.Empty();
+
+	UE_CLOG(NumSkippedPackages > 0, LogVirtualization, Warning, TEXT("Discarded %d duplicate package paths"), NumSkippedPackages);
+
 	// We need to reset the loader of any loaded package that should have its package file replaced
 	for (const TPair<FString, FString>& Pair : PackagesToReplace)
 	{
diff --git a/Engine/Source/Developer/Virtualization/Private/PackageVirtualizationProcess.cpp b/Engine/Source/Developer/Virtualization/Private/PackageVirtualizationProcess.cpp
--- a/Engine/Source/Developer/Virtualization/Private/PackageVirtualizationProcess.cpp
+++ b/Engine/Source/Developer/Virtualization/Private/PackageVirtualizationProcess.cpp
@@ -207,10 +207,21 @@ void VirtualizePackages(TConstArrayView<FString> PackagePaths, EVirtualizationOp
 	int64 TotalPackageTrailersFound = 0;
 	int64 TotalPayloadsToVirtualize = 0;
 
+	TSet<FString> ConsideredPackages;
+	ConsideredPackages.Reserve(PackagePaths.Num());
+
 	for (const FString& AbsoluteFilePath : PackagePaths)
 	{
 		FPackagePath PackagePath = FPackagePath::FromLocalPath(AbsoluteFilePath);
 
+		bool bIsDuplicate = false;
+		ConsideredPackages.Add(PackagePath.GetPackageName(), &bIsDuplicate);
+		if (bIsDuplicate)
+		{
+			UE_LOG(LogVirtualization, Verbose, TEXT("Skipping duplicate package entry '%s'"), *AbsoluteFilePath);
+			continue; // Skip duplicate packages
+		}
+
 		// TODO: How to handle text packages?
 		if (FPackageName::IsPackageExtension(PackagePath.GetHeaderExtension()) || FPackageName::IsTextPackageExtension(PackagePath.GetHeaderExtension()))
 		{
@@ -252,7 +263,11 @@ void VirtualizePackages(TConstArrayView<FString> PackagePaths, EVirtualizationOp
 		}
 	}
 
+	const int32 NumSkippedPackages = PackagePaths.Num() - ConsideredPackages.Num();
+	ConsideredPackages.Empty();
+
 	UE_LOG(LogVirtualization, Display, TEXT("Found %" INT64_FMT " package(s), %" INT64_FMT " of which had payload trailers"), TotalPackagesFound, TotalPackageTrailersFound);
+	UE_CLOG(NumSkippedPackages > 0, LogVirtualization, Warning, TEXT("Discarded %d duplicate package paths"), NumSkippedPackages);
 	UE_CLOG(TotalOutOfDatePackages > 0, LogVirtualization, Warning, TEXT("Found %" INT64_FMT " package(s) that are out of date and need resaving"), TotalOutOfDatePackages);
 
 	// TODO: Currently not all of the filtering is done as package save time, so some of the local payloads may not get virtualized.
