Add a way for VA backends to supply additional errors for failed pull operations when calling FVirtualizationManager::OnPayloadPullError

#rnx

- Each pull operation now takes a FText that can be filled in with additional info about any failures that will cause the new error dialog to be displayed.
-- At the moment we only have info to provide from the source control backend if the connection itself failed.
-- In the future we might want to completely overhaul our error logging and display all messages to the user but our logging system is not really built for that level of redirection at the moment.
- Only pass on the error if it came from a persistent backend, we don't need to worry the user about cached backends which are allowed to fail.

diff --git a/Engine/Source/Developer/Virtualization/Private/VirtualizationDDCBackend.cpp b/Engine/Source/Developer/Virtualization/Private/VirtualizationDDCBackend.cpp
--- a/Engine/Source/Developer/Virtualization/Private/VirtualizationDDCBackend.cpp
+++ b/Engine/Source/Developer/Virtualization/Private/VirtualizationDDCBackend.cpp
@@ -133,7 +133,7 @@ bool FDDCBackend::PushData(TArrayView<FPushRequest> Requests, EPushFlags Flags)
 	return bWasSuccess;
 }
 
-bool FDDCBackend::PullData(TArrayView<FPullRequest> Requests, EPullFlags Flags)
+bool FDDCBackend::PullData(TArrayView<FPullRequest> Requests, EPullFlags Flags, FText& OutErrors)
 {
 	TRACE_CPUPROFILER_EVENT_SCOPE(FDDCBackend::PullData);
 
diff --git a/Engine/Source/Developer/Virtualization/Private/VirtualizationDDCBackend.h b/Engine/Source/Developer/Virtualization/Private/VirtualizationDDCBackend.h
--- a/Engine/Source/Developer/Virtualization/Private/VirtualizationDDCBackend.h
+++ b/Engine/Source/Developer/Virtualization/Private/VirtualizationDDCBackend.h
@@ -39,7 +39,7 @@ private:
 	virtual EConnectionStatus OnConnect() override;
 
 	virtual bool PushData(TArrayView<FPushRequest> Requests, EPushFlags Flags) override;
-	virtual bool PullData(TArrayView<FPullRequest> Requests, EPullFlags Flags) override;
+	virtual bool PullData(TArrayView<FPullRequest> Requests, EPullFlags Flags, FText& OutErrors) override;
 	
 	virtual bool DoesPayloadExist(const FIoHash& Id) override;
 
diff --git a/Engine/Source/Developer/Virtualization/Private/VirtualizationFileBackend.cpp b/Engine/Source/Developer/Virtualization/Private/VirtualizationFileBackend.cpp
--- a/Engine/Source/Developer/Virtualization/Private/VirtualizationFileBackend.cpp
+++ b/Engine/Source/Developer/Virtualization/Private/VirtualizationFileBackend.cpp
@@ -179,7 +179,7 @@ bool FFileSystemBackend::PushData(TArrayView<FPushRequest> Requests, EPushFlags
 	return ErrorCount == 0;
 }
 
-bool FFileSystemBackend::PullData(TArrayView<FPullRequest> Requests, EPullFlags Flags)
+bool FFileSystemBackend::PullData(TArrayView<FPullRequest> Requests, EPullFlags Flags, FText& OutErrors)
 {
 	TRACE_CPUPROFILER_EVENT_SCOPE(FFileSystemBackend::PullData);
 
diff --git a/Engine/Source/Developer/Virtualization/Private/VirtualizationFileBackend.h b/Engine/Source/Developer/Virtualization/Private/VirtualizationFileBackend.h
--- a/Engine/Source/Developer/Virtualization/Private/VirtualizationFileBackend.h
+++ b/Engine/Source/Developer/Virtualization/Private/VirtualizationFileBackend.h
@@ -45,7 +45,7 @@ private:
 	virtual EConnectionStatus OnConnect() override;
 
 	virtual bool PushData(TArrayView<FPushRequest> Requests, EPushFlags Flags) override;
-	virtual bool PullData(TArrayView<FPullRequest> Requests, EPullFlags Flags) override;
+	virtual bool PullData(TArrayView<FPullRequest> Requests, EPullFlags Flags, FText& OutErrors) override;
 
 	virtual bool DoesPayloadExist(const FIoHash& Id) override;
 	
diff --git a/Engine/Source/Developer/Virtualization/Private/VirtualizationManager.cpp b/Engine/Source/Developer/Virtualization/Private/VirtualizationManager.cpp
--- a/Engine/Source/Developer/Virtualization/Private/VirtualizationManager.cpp
+++ b/Engine/Source/Developer/Virtualization/Private/VirtualizationManager.cpp
@@ -732,8 +732,9 @@ bool FVirtualizationManager::PushData(TArrayView<FPushRequest> Requests, EStorag
 		{
 			for (FPushRequest& PushRequest : ValidatedRequests)
 			{
+				FText Errors;
 				FPullRequest PullRequest(PushRequest.GetIdentifier());
-				PullDataFromBackend(*Backend, MakeArrayView(&PullRequest, 1));
+				PullDataFromBackend(*Backend, MakeArrayView(&PullRequest, 1), Errors);
 
 				checkf(PushRequest.GetIdentifier() == PullRequest.GetPayload().GetRawHash(),
 						TEXT("[%s] Failed to pull payload '%s' after it was pushed to backend"),
@@ -1713,6 +1714,11 @@ void FVirtualizationManager::AddBackend(TUniquePtr<IVirtualizationBackend> Backe
 	COOK_STAT(Profiling::CreateStats(*BackendRef));
 }
 
+bool FVirtualizationManager::IsPersistentBackend(IVirtualizationBackend& Backend)
+{
+	return PersistentStorageBackends.Contains(&Backend);
+}
+
 void FVirtualizationManager::EnsureBackendConnections()
 {
 	if (bPendingBackendConnections)
@@ -1763,8 +1769,9 @@ void FVirtualizationManager::CachePayloads(TArrayView<FPushRequest> Requests, co
 		{
 			for (const FPushRequest& Request : Requests)
 			{
+				FText Errors;
 				FPullRequest ValidationRequest(Request.GetIdentifier());
-				PullDataFromBackend(*BackendToCache, MakeArrayView(&ValidationRequest, 1));
+				PullDataFromBackend(*BackendToCache, MakeArrayView(&ValidationRequest, 1), Errors);
 
 				checkf(Request.GetPayload().GetRawHash() == ValidationRequest.GetPayload().GetRawHash(),
 					TEXT("[%s] Failed to pull payload '%s' after it was cached to backend"),
@@ -1853,6 +1860,8 @@ void FVirtualizationManager::PullDataFromAllBackends(TArrayView<FPullRequest> Re
 	
 	while(true)
 	{
+		TStringBuilder<512> BackendErrors;
+
 		for (IVirtualizationBackend* Backend : PullEnabledBackends)
 		{
 			check(Backend != nullptr);
@@ -1869,7 +1878,19 @@ void FVirtualizationManager::PullDataFromAllBackends(TArrayView<FPullRequest> Re
 				continue;
 			}
 
-			PullDataFromBackend(*Backend, RequestsCollection.GetRequests());
+			FText Errors;
+			PullDataFromBackend(*Backend, RequestsCollection.GetRequests(), Errors);
+
+			// We only want to report errors from persistent backends, cached backends are allowed to fail.
+			if (!Errors.IsEmpty() && IsPersistentBackend(*Backend))
+			{
+				if (BackendErrors.Len() != 0)
+				{
+					BackendErrors << LINE_TERMINATOR;
+				}
+
+				BackendErrors << Backend->GetDebugName() << TEXT(": ") << Errors.ToString();
+			}
 
 			const bool bShouldCache = EnumHasAllFlags(CachingPolicy, ECachingPolicy::CacheOnPull);
 
@@ -1893,20 +1914,20 @@ void FVirtualizationManager::PullDataFromAllBackends(TArrayView<FPullRequest> Re
 		{
 			return; // All payloads pulled
 		}
-		else if (OnPayloadPullError() != ErrorHandlingResult::Retry)
+		else if (OnPayloadPullError(BackendErrors) != ErrorHandlingResult::Retry)
 		{
 			return; // Some payloads failed to pull
 		}
 	}
 }
 
-void FVirtualizationManager::PullDataFromBackend(IVirtualizationBackend& Backend, TArrayView<FPullRequest> Requests)
+void FVirtualizationManager::PullDataFromBackend(IVirtualizationBackend& Backend, TArrayView<FPullRequest> Requests, FText& OutErrors)
 {
 	COOK_STAT(FCookStats::CallStats & Stats = Profiling::GetPullStats(Backend));
 	COOK_STAT(FCookStats::FScopedStatsCounter Timer(Stats));
 	COOK_STAT(Timer.TrackCyclesOnly());
 	
-	Backend.PullData(Requests, IVirtualizationBackend::EPullFlags::None);
+	Backend.PullData(Requests, IVirtualizationBackend::EPullFlags::None, OutErrors);
 	
 #if ENABLE_COOK_STATS
 	const bool bIsInGameThread = IsInGameThread();
@@ -1924,7 +1945,7 @@ void FVirtualizationManager::PullDataFromBackend(IVirtualizationBackend& Backend
 #endif //ENABLE_COOK_STATS
 }
 
-FVirtualizationManager::ErrorHandlingResult FVirtualizationManager::OnPayloadPullError()
+FVirtualizationManager::ErrorHandlingResult FVirtualizationManager::OnPayloadPullError(FStringView BackendErrors)
 {
 	if (bUseLegacyErrorHandling)
 	{
@@ -1941,6 +1962,12 @@ FVirtualizationManager::ErrorHandlingResult FVirtualizationManager::OnPayloadPul
 		FTextBuilder MsgBuilder;
 		MsgBuilder.AppendLine(LOCTEXT("VAPullMsgHeader", "Failed to pull payload(s) from virtualization storage and allowing the editor to continue could corrupt data!"));
 		
+		if (!BackendErrors.IsEmpty())
+		{
+			MsgBuilder.AppendLine(FString(TEXT("")));
+			MsgBuilder.AppendLine(FString(BackendErrors));
+		}
+
 		if (!PullErrorAdditionalMsg.IsEmpty())
 		{
 			MsgBuilder.AppendLine(FString(TEXT("")));
diff --git a/Engine/Source/Developer/Virtualization/Private/VirtualizationManager.h b/Engine/Source/Developer/Virtualization/Private/VirtualizationManager.h
--- a/Engine/Source/Developer/Virtualization/Private/VirtualizationManager.h
+++ b/Engine/Source/Developer/Virtualization/Private/VirtualizationManager.h
@@ -215,6 +215,8 @@ private:
 
 	void AddBackend(TUniquePtr<IVirtualizationBackend> Backend, FBackendArray& PushArray);
 
+	bool IsPersistentBackend(IVirtualizationBackend& Backend);
+
 	void EnsureBackendConnections();
 
 	void CachePayloads(TArrayView<FPushRequest> Requests, const IVirtualizationBackend* BackendSource, IVirtualizationBackend::EPushFlags Flags);
@@ -223,7 +225,7 @@ private:
 	bool TryPushDataToBackend(IVirtualizationBackend& Backend, TArrayView<FPushRequest> Requests);
 
 	void PullDataFromAllBackends(TArrayView<FPullRequest> Requests);
-	void PullDataFromBackend(IVirtualizationBackend& Backend, TArrayView<FPullRequest> Requests);
+	void PullDataFromBackend(IVirtualizationBackend& Backend, TArrayView<FPullRequest> Requests, FText& OutErrors);
 
 	enum class ErrorHandlingResult
 	{
@@ -233,7 +235,7 @@ private:
 		AcceptFailedPayloads
 	};
 
-	ErrorHandlingResult OnPayloadPullError();
+	ErrorHandlingResult OnPayloadPullError(FStringView BackendErrors);
 	
 	bool ShouldVirtualizeAsset(const UObject* Owner) const;
 
diff --git a/Engine/Source/Developer/Virtualization/Private/VirtualizationSourceControlBackend.cpp b/Engine/Source/Developer/Virtualization/Private/VirtualizationSourceControlBackend.cpp
--- a/Engine/Source/Developer/Virtualization/Private/VirtualizationSourceControlBackend.cpp
+++ b/Engine/Source/Developer/Virtualization/Private/VirtualizationSourceControlBackend.cpp
@@ -403,7 +403,7 @@ IVirtualizationBackend::EConnectionStatus FSourceControlBackend::OnConnect()
 	return IVirtualizationBackend::EConnectionStatus::Connected;
 }
 
-bool FSourceControlBackend::PullData(TArrayView<FPullRequest> Requests, EPullFlags Flags)
+bool FSourceControlBackend::PullData(TArrayView<FPullRequest> Requests, EPullFlags Flags, FText& OutErrors)
 {
 	TRACE_CPUPROFILER_EVENT_SCOPE(FSourceControlBackend::PullData);
 
@@ -425,6 +425,8 @@ bool FSourceControlBackend::PullData(TArrayView<FPullRequest> Requests, EPullFla
 
 	int32 Retries = 0;
 
+	bool bConnectionFailed = false;
+
 	while (Retries < RetryCount)
 	{
 		// Only warn if the backend is configured to retry
@@ -463,9 +465,20 @@ bool FSourceControlBackend::PullData(TArrayView<FPullRequest> Requests, EPullFla
 			return false;
 		}
 
+		bConnectionFailed = DownloadCommand->GetResultInfo().DidConnectionFail();
+
 		Retries++;
 	}
 
+	if (bConnectionFailed)
+	{
+		TMap<ISourceControlProvider::EStatus, FString> SCPStatus = SCCProvider->GetStatus();
+		
+		OutErrors = FText::Format(LOCTEXT("VA_SCP", "Failed to connect to perforce server '{0}' with username '{1}'"),
+			FText::FromString(SCPStatus[ISourceControlProvider::EStatus::Port]), 
+			FText::FromString(SCPStatus[ISourceControlProvider::EStatus::User]));
+	}
+
 	return false;
 }
 
diff --git a/Engine/Source/Developer/Virtualization/Private/VirtualizationSourceControlBackend.h b/Engine/Source/Developer/Virtualization/Private/VirtualizationSourceControlBackend.h
--- a/Engine/Source/Developer/Virtualization/Private/VirtualizationSourceControlBackend.h
+++ b/Engine/Source/Developer/Virtualization/Private/VirtualizationSourceControlBackend.h
@@ -128,7 +128,7 @@ private:
 	virtual EConnectionStatus OnConnect() override;
 	
 	virtual bool PushData(TArrayView<FPushRequest> Requests, EPushFlags Flags) override;
-	virtual bool PullData(TArrayView<FPullRequest> Requests, EPullFlags Flags) override;
+	virtual bool PullData(TArrayView<FPullRequest> Requests, EPullFlags Flags, FText& OutErrors) override;
 
 	virtual bool DoesPayloadExist(const FIoHash& Id) override;
 	
diff --git a/Engine/Source/Developer/Virtualization/Public/IVirtualizationBackend.h b/Engine/Source/Developer/Virtualization/Public/IVirtualizationBackend.h
--- a/Engine/Source/Developer/Virtualization/Public/IVirtualizationBackend.h
+++ b/Engine/Source/Developer/Virtualization/Public/IVirtualizationBackend.h
@@ -133,6 +133,9 @@ public:
 	 * backend to do this.
 	 * 
 	 * @param	Requests	An array of payload pull requests. @see FPullRequest
+	 * @param	Flags		Flags that allow the pull to be specialized @see EPullFlags
+	 * @param	OutErrors	If the pull fails, the backend can add optional additional error messages
+	 *						to be displayed directly to the user.
 	 * 
 	 * @return				True if no errors were encountered while pulling, otherwise false.
 	 *						Note that returning true does not mean that all of the payloads were
@@ -142,7 +145,7 @@ public:
 	// Assume that the array only has unique requests
 	// Will set request error value if there is a problem with loading it
 	// Returns false on critical error, otherwise true
-	virtual bool PullData(TArrayView<FPullRequest> Requests, EPullFlags Flags) = 0;
+	virtual bool PullData(TArrayView<FPullRequest> Requests, EPullFlags Flags, FText& OutErrors) = 0;
 	
 	/**
 	 * Checks if a payload exists in the backends storage.
@@ -284,7 +287,7 @@ private:
  * IModularFeatures with the feature name "VirtualizationBackendFactory" to
  * give 'FVirtualizationManager' access to it. 
  * The macro 'UE_REGISTER_VIRTUALIZATION_BACKEND_FACTORY' can be used to create 
- * a factory easily if you do not want to specialize the behaviour.
+ * a factory easily if you do not want to specialize the behavior.
 */
 class IVirtualizationBackendFactory : public IModularFeature
 {
