Do not assign payloads downloaded from source control if they did not actually download.

#rnx

- Ends up pretty much the same,as FCompressedBuffer::FromCompressed will return a Null FCompressedBuffer if the provided data is null, but adding the check makes it clear when reading FSourceControlBackend::PullData that payloads might not be found and we are not setting them in that case.

diff --git a/Engine/Source/Developer/Virtualization/Private/VirtualizationSourceControlBackend.cpp b/Engine/Source/Developer/Virtualization/Private/VirtualizationSourceControlBackend.cpp
--- a/Engine/Source/Developer/Virtualization/Private/VirtualizationSourceControlBackend.cpp
+++ b/Engine/Source/Developer/Virtualization/Private/VirtualizationSourceControlBackend.cpp
@@ -450,7 +450,10 @@ bool FSourceControlBackend::PullData(TArrayView<FPullRequest> Requests, EPullFla
 		{
 			// The payload was created by FCompressedBuffer::Compress so we can return it as a FCompressedBuffer.
 			FSharedBuffer Buffer = DownloadCommand->GetFileData(DepotPaths[Index]);
-			Requests[Index].SetPayload(FCompressedBuffer::FromCompressed(Buffer));
+			if (!Buffer.IsNull())
+			{
+				Requests[Index].SetPayload(FCompressedBuffer::FromCompressed(Buffer));
+			}
 		}
 
 		if (bOperationSuccess)
