Enable UnsafeTypeCastWarningLevel errors for the virtualization module

#rnx

- I thought I had already done this but I see nothing in the p4 history, so it is likely that I forgot to actually submit the change.
- Fix some additional warnings found under clang (C5219)

diff --git a/Engine/Source/Developer/Virtualization/Private/DumpPackagePayloadInfoCommand.cpp b/Engine/Source/Developer/Virtualization/Private/DumpPackagePayloadInfoCommand.cpp
--- a/Engine/Source/Developer/Virtualization/Private/DumpPackagePayloadInfoCommand.cpp
+++ b/Engine/Source/Developer/Virtualization/Private/DumpPackagePayloadInfoCommand.cpp
@@ -22,12 +22,12 @@ FString BytesToString(int64 SizeInBytes)
 	}
 	else if (SizeInBytes < (1024 * 1024))
 	{
-		double SizeInKb = SizeInBytes / (1024.0);
+		double SizeInKb = static_cast<double>(SizeInBytes) / (1024.0);
 		return FString::Printf(TEXT("%.2f KB"), SizeInKb);
 	}
 	else
 	{
-		double SizeInMB = SizeInBytes / (1024.0 * 1024.0);
+		double SizeInMB = static_cast<double>(SizeInBytes) / (1024.0 * 1024.0);
 		return FString::Printf(TEXT("%.2f MB"), SizeInMB);
 	}
 }
diff --git a/Engine/Source/Developer/Virtualization/Private/VirtualizationFileBackend.cpp b/Engine/Source/Developer/Virtualization/Private/VirtualizationFileBackend.cpp
--- a/Engine/Source/Developer/Virtualization/Private/VirtualizationFileBackend.cpp
+++ b/Engine/Source/Developer/Virtualization/Private/VirtualizationFileBackend.cpp
@@ -268,7 +268,7 @@ TUniquePtr<FArchive> FFileSystemBackend::OpenFileForReading(const TCHAR* FilePat
 		else
 		{
 			UE_LOG(LogVirtualization, Warning, TEXT("[%s] Failed to open '%s' for reading attempt retrying (%d/%d) in %dms..."), *GetDebugName(), FilePath, Retries, RetryCount, RetryWaitTimeMS);
-			FPlatformProcess::SleepNoStats(RetryWaitTimeMS * 0.001f);
+			FPlatformProcess::SleepNoStats(static_cast<float>(RetryWaitTimeMS) * 0.001f);
 
 			Retries++;
 		}
diff --git a/Engine/Source/Developer/Virtualization/Private/VirtualizationManager.cpp b/Engine/Source/Developer/Virtualization/Private/VirtualizationManager.cpp
--- a/Engine/Source/Developer/Virtualization/Private/VirtualizationManager.cpp
+++ b/Engine/Source/Developer/Virtualization/Private/VirtualizationManager.cpp
@@ -458,8 +458,8 @@ namespace Profiling
 					for (const auto& Iterator : Stats)
 					{
 						const int64 Count = Iterator.Value.GetAccumulatedValueAnyThread(FCookStats::CallStats::EHitOrMiss::Hit, FCookStats::CallStats::EStatType::Counter);
-						const double Time = Iterator.Value.GetAccumulatedValueAnyThread(FCookStats::CallStats::EHitOrMiss::Hit, FCookStats::CallStats::EStatType::Cycles) * FPlatformTime::GetSecondsPerCycle();
-						const double DataSizeMB = Iterator.Value.GetAccumulatedValueAnyThread(FCookStats::CallStats::EHitOrMiss::Hit, FCookStats::CallStats::EStatType::Bytes) / (1024.0f * 1024.0f);
+						const double Time = (double)Iterator.Value.GetAccumulatedValueAnyThread(FCookStats::CallStats::EHitOrMiss::Hit, FCookStats::CallStats::EStatType::Cycles) * FPlatformTime::GetSecondsPerCycle();
+						const double DataSizeMB = (double)Iterator.Value.GetAccumulatedValueAnyThread(FCookStats::CallStats::EHitOrMiss::Hit, FCookStats::CallStats::EStatType::Bytes) / (1024.0f * 1024.0f);
 						const double MBps = Time != 0.0 ? (DataSizeMB / Time) : 0.0;
 
 						UE_LOG(LogVirtualization, Display, TEXT("%-40.40s|%10lld|%17.1f|%12.3f|%14.3f|"),
diff --git a/Engine/Source/Developer/Virtualization/Private/VirtualizationSourceControlBackend.cpp b/Engine/Source/Developer/Virtualization/Private/VirtualizationSourceControlBackend.cpp
--- a/Engine/Source/Developer/Virtualization/Private/VirtualizationSourceControlBackend.cpp
+++ b/Engine/Source/Developer/Virtualization/Private/VirtualizationSourceControlBackend.cpp
@@ -433,7 +433,7 @@ bool FSourceControlBackend::PullData(TArrayView<FPullRequest> Requests, EPullFla
 		if (Retries != 0)
 		{
 		//	UE_LOG(LogVirtualization, Warning, TEXT("[%s] Failed to download '%s' retrying (%d/%d) in %dms..."), *GetDebugName(), DepotPath.ToString(), Retries, RetryCount, RetryWaitTimeMS);
-			FPlatformProcess::SleepNoStats(RetryWaitTimeMS * 0.001f);
+			FPlatformProcess::SleepNoStats(static_cast<float>(RetryWaitTimeMS) * 0.001f);
 		}
 
 #if IS_SOURCE_CONTROL_THREAD_SAFE
diff --git a/Engine/Source/Developer/Virtualization/Virtualization.Build.cs b/Engine/Source/Developer/Virtualization/Virtualization.Build.cs
--- a/Engine/Source/Developer/Virtualization/Virtualization.Build.cs
+++ b/Engine/Source/Developer/Virtualization/Virtualization.Build.cs
@@ -7,6 +7,8 @@ public class Virtualization : ModuleRules
 {
 	public Virtualization(ReadOnlyTargetRules Target) : base(Target)
 	{
+		UnsafeTypeCastWarningLevel = WarningLevel.Error;
+
 		PrivateIncludePathModuleNames.Add("Analytics");
 
 		PrivateDependencyModuleNames.AddRange(
