Allow an optional url to be shown along side VA connection errors to make it easier for users to find the help that they need.

#rnx

- A url can now be provided via 'ini:Engine:[Core.VirtualizationModule]:ConnectionHelpUrl=""' that will be added to various connection errors. If the error message is via FMessageLog then the user will be able to click the url to go straight to the link.
- This sort of help will be very company specific, this url will allow teams to build up their own FAQ for their users.
- Currently used by the source control backend if the initial connection fails but will be expanded to other areas.
- The url is currently parsed and owned by the virtualization manager but at the moment the connection error is issued by each backend. So right now there is a static accessor allowing backends to get the value to use. This would not work properly for 3rd party backends and will  be changed in the near term.

diff --git a/Engine/Config/BaseEngine.ini b/Engine/Config/BaseEngine.ini
--- a/Engine/Config/BaseEngine.ini
+++ b/Engine/Config/BaseEngine.ini
@@ -1405,6 +1405,7 @@ LazyInitConnections=False
 UseLegacyErrorHandling=False
 ForceCachingOnPull=False
 PullErrorAdditionalMsg=""
+ConnectionHelpUrl=""
 
 [/Script/Virtualization.VirtualizationFilterSettings]
 !ExcludePackagePaths=ClearArray
@@ -3319,4 +3320,4 @@ MaterialParents=(("MI_Default_Opaque", "/Interchange/gltf/MaterialInstances/MI_D
 
 [BuildPatchServices]
 ; Enable/Disable mechanism that provides a disk overflow store for patch chunks.
-bEnableDiskOverflowStore=true
+bEnableDiskOverflowStore=true
\ No newline at end of file
diff --git a/Engine/Source/Developer/Virtualization/Private/VirtualizationManager.cpp b/Engine/Source/Developer/Virtualization/Private/VirtualizationManager.cpp
--- a/Engine/Source/Developer/Virtualization/Private/VirtualizationManager.cpp
+++ b/Engine/Source/Developer/Virtualization/Private/VirtualizationManager.cpp
@@ -481,6 +481,8 @@ namespace Profiling
 #endif // ENABLE_COOK_STATS
 } //namespace Profiling
 
+FString FVirtualizationManager::ConnectionHelpUrl;
+
 FVirtualizationManager::FVirtualizationManager()
 	: bAllowPackageVirtualization(true)
 	, CachingPolicy(ECachingPolicy::AlwaysCache)
@@ -1134,6 +1136,11 @@ void FVirtualizationManager::ApplySettingsFromConfigFiles(const FConfigFile& Con
 		ConfigFile.GetString(ConfigSection, TEXT("PullErrorAdditionalMsg"), PullErrorAdditionalMsg);
 		// This value is not echoed to the log file, as seeing an error string there might confuse users
 	}
+	
+	{
+		ConfigFile.GetString(ConfigSection, TEXT("ConnectionHelpUrl"), ConnectionHelpUrl);
+		// This value is not echoed to the log file
+	}
 
 	{
 		ConfigFile.GetBool(ConfigSection, TEXT("ForceCachingOnPull"), bForceCachingOnPull);
@@ -2197,6 +2204,11 @@ void FVirtualizationManager::GatherAnalytics(TArray<FAnalyticsEventAttribute>& A
 	}
 }
 
+FString FVirtualizationManager::GetConnectionHelpUrl()
+{
+	return ConnectionHelpUrl;
+}
+
 } // namespace UE::Virtualization
 
 #undef UE_INLINE_ALLOCATION_COUNT
diff --git a/Engine/Source/Developer/Virtualization/Private/VirtualizationManager.h b/Engine/Source/Developer/Virtualization/Private/VirtualizationManager.h
--- a/Engine/Source/Developer/Virtualization/Private/VirtualizationManager.h
+++ b/Engine/Source/Developer/Virtualization/Private/VirtualizationManager.h
@@ -321,6 +321,8 @@ private:
 	/** An additional error message to display when pulling payloads fails */
 	FString PullErrorAdditionalMsg;
 
+	/** Optional url used to augment connection failure error messages */
+	static FString ConnectionHelpUrl;
 private:
 
 	/** The name of the current project */
@@ -385,6 +387,13 @@ private:
 		/** The number of upcoming payload pulls that should be failed */
 		std::atomic<int32> MissCount = 0;
 	} DebugValues;
+
+public:
+	/**
+	 * Temp accessor for backends until we add a better system for connection failure notification. Note that
+	 * this is not exposed outside of the module so that we can change/remove this easily in the future.
+	 */
+	static FString GetConnectionHelpUrl();
 };
 
 } // namespace UE::Virtualization
diff --git a/Engine/Source/Developer/Virtualization/Private/VirtualizationSourceControlBackend.cpp b/Engine/Source/Developer/Virtualization/Private/VirtualizationSourceControlBackend.cpp
--- a/Engine/Source/Developer/Virtualization/Private/VirtualizationSourceControlBackend.cpp
+++ b/Engine/Source/Developer/Virtualization/Private/VirtualizationSourceControlBackend.cpp
@@ -8,6 +8,7 @@
 #include "IO/IoHash.h"
 #include "ISourceControlModule.h"
 #include "ISourceControlProvider.h"
+#include "Interfaces/IPluginManager.h"
 #include "Logging/MessageLog.h"
 #include "Misc/FileHelper.h"
 #include "Misc/Parse.h"
@@ -16,9 +17,9 @@
 #include "Misc/ScopeExit.h"
 #include "SourceControlInitSettings.h"
 #include "SourceControlOperations.h"
+#include "VirtualizationManager.h"
 #include "VirtualizationSourceControlUtilities.h"
 #include "VirtualizationUtilities.h"
-#include "Interfaces/IPluginManager.h"
 
 // When the SourceControl module (or at least the perforce source control module) is thread safe we
 // can enable this and stop using the hacky work around 'TryToDownloadFileFromBackgroundThread'
@@ -1112,48 +1113,61 @@ void FSourceControlBackend::OnConnectionError(FText Message)
 	// GameThread and at a point where notification will work. In addition if we do need
 	// to defer the error message until next tick (rather than just the notification) then
 	// we will write it to the log at the point it is raised rather than the point it is
-	// defered to in an attempt to make the log more readable.
-	if (!Message.IsEmpty())
+	// deferred to in an attempt to make the log more readable.
+
+	if (Message.IsEmpty())
 	{
-		if (::IsInGameThread())
-		{
-			// If we are on the game thread we can post the error message immediately
-			FMessageLog Log("LogVirtualization");
-			Log.Error(Message);
+		return;
+	}
 
-			// Clear the message so we don't log it a second time
-			Message = FText();
-		}
-		else
-		{
-			// We can only send a FMessageLog on the GameThread so for now just log the error
-			// and we can send it to the FMessageLog system next tick
-			UE_LOG(LogVirtualization, Error, TEXT("%s"), *Message.ToString());
-		}
+	TSharedRef<FTokenizedMessage> TokenizedMsg = FTokenizedMessage::Create(EMessageSeverity::Error);
+	TokenizedMsg->AddToken(FTextToken::Create(Message));
+
+	FString ConnectionHelpUrl = FVirtualizationManager::GetConnectionHelpUrl();
+	if (!ConnectionHelpUrl.IsEmpty())
+	{
+		TokenizedMsg->AddToken(FURLToken::Create(ConnectionHelpUrl, LOCTEXT("URL", "Additional connection help")));
 	}
 
-	if (!bSuppressNotifications || !Message.IsEmpty())
+	if (::IsInGameThread())
+	{
+		// If we are on the game thread we can post the error message immediately
+		FMessageLog Log("LogVirtualization");
+		Log.AddMessage(TokenizedMsg);
+	}
+	else
 	{
-		auto Callback = [Message, bShouldNotify = bSuppressNotifications](float Delta)->bool
+		// We can only send a FMessageLog on the GameThread so for now just log the error
+		// and we can send it to the FMessageLog system next tick
+		UE_LOG(LogVirtualization, Error, TEXT("%s"), *Message.ToString());
+
+		auto Callback = [TokenizedMsg, bShouldNotify = !bSuppressNotifications](float Delta)->bool
 		{
 			FMessageLog Log("LogVirtualization");
 			Log.SuppressLoggingToOutputLog();
+			Log.AddMessage(TokenizedMsg);
 
-			if (!Message.IsEmpty())
-			{
-				Log.Error(Message);
-			}
+			return false;
+		};
 
-			if (!bShouldNotify)
+		FTSTicker::GetCoreTicker().AddTicker(FTickerDelegate::CreateLambda(MoveTemp(Callback)));
+	}
+
+	if (!bSuppressNotifications)
+	{
+		// Add the notification to the ticker so that if this is called during editor startup it will only
+		// be issued on the first frame, increasing the chance that the user will see it.
+		auto NotificationCallback = [](float Delta)->bool
 			{
-				Log.Notify(LOCTEXT("ConnectionError", "Asset virtualization connection errors were encountered, see the message log for more info"));
-			}
+				FMessageLog Notification("LogVirtualization");;
+				Notification.SuppressLoggingToOutputLog();
 
-			// This tick callback is one shot, so return false to prevent it being invoked again
-			return false;
-		}; 
-	
-		FTSTicker::GetCoreTicker().AddTicker(FTickerDelegate::CreateLambda(MoveTemp(Callback)));
+				Notification.Notify(LOCTEXT("ConnectionError", "Asset virtualization connection errors were encountered, see the message log for more info"));
+
+				return false;
+			};
+
+		FTSTicker::GetCoreTicker().AddTicker(FTickerDelegate::CreateLambda(MoveTemp(NotificationCallback)));
 	}
 }
 
