Change TryCopyPackageWithoutTrailer to take the length of the trailer to remove rather than the trailer structure itself.

#'rnx

diff --git a/Engine/Source/Developer/Virtualization/Private/PackageUtils.cpp b/Engine/Source/Developer/Virtualization/Private/PackageUtils.cpp
--- a/Engine/Source/Developer/Virtualization/Private/PackageUtils.cpp
+++ b/Engine/Source/Developer/Virtualization/Private/PackageUtils.cpp
@@ -53,7 +53,7 @@ bool CanWriteToFile(const FString& FilePath)
 	return FileHandle.IsValid();
 }
 
-bool TryCopyPackageWithoutTrailer(const FString& SourcePath, const FString& DstPath, const FPackageTrailer& Trailer, TArray<FText>& OutErrors)
+bool TryCopyPackageWithoutTrailer(const FString& SourcePath, const FString& DstPath, int64 TrailerLength, TArray<FText>& OutErrors)
 {
 	// TODO: Consider adding a custom copy routine to only copy the data we want, rather than copying the full file then truncating
 
@@ -66,7 +66,7 @@ bool TryCopyPackageWithoutTrailer(const FString& SourcePath, const FString& DstP
 		return false;
 	}
 
-	const int64 PackageSizeWithoutTrailer = IFileManager::Get().FileSize(*SourcePath) - Trailer.GetTrailerLength();
+	const int64 PackageSizeWithoutTrailer = IFileManager::Get().FileSize(*SourcePath) - TrailerLength;
 
 	{
 		TUniquePtr<IFileHandle> TempFileHandle(FPlatformFileManager::Get().GetPlatformFile().OpenWrite(*DstPath, true));
@@ -109,7 +109,7 @@ FString DuplicatePackageWithUpdatedTrailer(const FString& AbsolutePackagePath, c
 	// TODO Optimization: Combine TryCopyPackageWithoutTrailer with the appending of the new trailer to avoid opening multiple handles
 
 	// Create copy of package minus the trailer the trailer
-	if (!TryCopyPackageWithoutTrailer(AbsolutePackagePath, TempFilePath, Trailer, OutErrors))
+	if (!TryCopyPackageWithoutTrailer(AbsolutePackagePath, TempFilePath, Trailer.GetTrailerLength(), OutErrors))
 	{
 		return FString();
 	}
@@ -157,7 +157,7 @@ FString DuplicatePackageWithNewTrailer(const FString& AbsolutePackagePath, const
 	// TODO Optimization: Combine TryCopyPackageWithoutTrailer with the appending of the new trailer to avoid opening multiple handles
 
 	// Create copy of package minus the trailer the trailer
-	if (!TryCopyPackageWithoutTrailer(AbsolutePackagePath, TempFilePath, OriginalTrailer, OutErrors))
+	if (!TryCopyPackageWithoutTrailer(AbsolutePackagePath, TempFilePath, OriginalTrailer.GetTrailerLength(), OutErrors))
 	{
 		return FString();
 	}
diff --git a/Engine/Source/Developer/Virtualization/Private/PackageUtils.h b/Engine/Source/Developer/Virtualization/Private/PackageUtils.h
--- a/Engine/Source/Developer/Virtualization/Private/PackageUtils.h
+++ b/Engine/Source/Developer/Virtualization/Private/PackageUtils.h
@@ -39,12 +39,12 @@ bool CanWriteToFile(const FString& FilePath);
  *
  * @param SourcePath	The absolute path of the package to copy
  * @param DstPath		The path where the copy should be created
- * @param Trailer		The trailer found in 'PackagePath' that is already loaded (used to know where we cut the trailer off)
+ * @param TrailerLength	The length of the trailer to be removed
  * @param Errors [out]	Errors created by the function will be added here
  *
  * @return Returns true if the package was copied correctly, false otherwise. Note even when returning false a file might have been created at 'CopyPath'
  */
-bool TryCopyPackageWithoutTrailer(const FString& SourcePath, const FString& DstPath, const FPackageTrailer& Trailer, TArray<FText>& OutErrors);
+bool TryCopyPackageWithoutTrailer(const FString& SourcePath, const FString& DstPath, int64 TrailerLength, TArray<FText>& OutErrors);
 
 /**
  * Create a copy of an existing package but replace the trailer with an updated version of the original trailer. The copy will be to a tmp file
