Small optimization when rehydrating packages containing multiple payloads.

- Reduced time taken to hydrate a 16GB test project by ~20%.
-- Still more work to do, this was just very easy low hanging fruit.
- Only really helps if packages contain multiple virtualized payloads.

diff --git a/Engine/Source/Developer/Virtualization/Private/PackageRehydrationProcess.cpp b/Engine/Source/Developer/Virtualization/Private/PackageRehydrationProcess.cpp
--- a/Engine/Source/Developer/Virtualization/Private/PackageRehydrationProcess.cpp
+++ b/Engine/Source/Developer/Virtualization/Private/PackageRehydrationProcess.cpp
@@ -140,29 +140,33 @@ bool TryRehydrateBuilder(const FString& FilePath, FPackageTrailer& OutTrailer, F
 
 	IVirtualizationSystem& System = IVirtualizationSystem::Get();
 
+	TArray<FPullRequest> Requests;
+	Requests.Reserve(VirtualizedPayloads.Num());
+
 	for (const FIoHash& Id : VirtualizedPayloads)
 	{
-		FCompressedBuffer Payload = System.PullData(Id);
-		if (!Payload.IsNull())
-		{
-			if (OutBuilder.UpdatePayloadAsLocal(Id, Payload))
-			{
-				PayloadsHydrated++;
-			}
-			else
-			{
-				FText Message = FText::Format(LOCTEXT("VAHydration_UpdateStatusFailed", "Unable to update the status for the payload '{0}' in the package '{1}'"),
-					FText::FromString(LexToString(Id)),
-					FText::FromString(FilePath));
-				OutErrors.Add(Message);
+		Requests.Emplace(FPullRequest(Id));
+	}
 
-				return false;
-			}
+	if (!System.PullData(Requests))
+	{
+		FText Message = FText::Format(LOCTEXT("VAHydration_PullFailed", "Unable to pull the data for the package '{0}'"),
+			FText::FromString(FilePath));
+		OutErrors.Add(Message);
+
+		return false;
+	}
+
+	for (const FPullRequest& Request : Requests)
+	{
+		if (OutBuilder.UpdatePayloadAsLocal(Request.GetIdentifier(), Request.GetPayload()))
+		{
+			PayloadsHydrated++;
 		}
 		else
 		{
-			FText Message = FText::Format(LOCTEXT("VAHydration_PullFailed", "Unable to pull the data for the payload '{0}' for the package '{1}'"),
-				FText::FromString(LexToString(Id)),
+			FText Message = FText::Format(LOCTEXT("VAHydration_UpdateStatusFailed", "Unable to update the status for the payload '{0}' in the package '{1}'"),
+				FText::FromString(LexToString(Request.GetIdentifier())),
 				FText::FromString(FilePath));
 			OutErrors.Add(Message);
 
@@ -170,6 +174,8 @@ bool TryRehydrateBuilder(const FString& FilePath, FPackageTrailer& OutTrailer, F
 		}
 	}
 
+	check(OutBuilder.GetNumVirtualizedPayloads() == 0);
+
 	return PayloadsHydrated > 0;
 }
 
