Now if the VA source control backend fails to connect we can display a dialog to the user asking for the correct credentials directly. Currently disabled by default.

### Dialog
- When the source control backend tries to connect and fails we will now show a dialog allowing the user to try new settings rather than let the editor continue and crash.
- The dialog will allow the user to enter the server address (P4PORT) and their username (P4USER).
- Once the dialog is displayed the user will have three options to continue.
-- "Reset To Defaults" which will remove all settings saved to the users local ini files and attempt to use the environment defaults to connect.
-- "Retry Connection" which will take the new address and username that the user provided and attempt to connect with those settings.
-- "Skip" no connection attempt will be made. This may cause instability later if the user needs data that they cannot access.
- If the connection succeeds after the user has reset to defaults or provided new settings, then they will be saved in the local ini file for future sessions.
- The dialog will not show if:
-- Slate is disabled for the current processor has not yet initialized
-- The error is found not on the game thread. This means that lazy init will not currently work well with this dialog. Attempts to mashal the error to the game thread have been spotty at best and there is a high chance of thread locking so attempts were abandoned.
-- If the -unattended flag is set on the process.
-- If the dialog is disabled via the config file when setting up the backend (UseRetryConnectionDialog=False)
-- If  'engine.ini:[Core.VirtualizationModule]:UseLocalIniFileSettings' is false as future sessions would not be able to use the new settings and would become tiresome. In these cases the environment settings will need to be fixed.
- If the dialog is now shown then the backend is marked as unconnected and will not work.

### Known Issues
- We currently cannot get the settings that the failed p4 connection used, nor the error message dueto how the API works. So for now we just display an empty server address and username edit box. This will be fixed when the API is changed.
- There is currently no good way to inform the user where their settings come from.
- There is no easy way to reset the settings once saved. To allow the user to change the server address to a better proxy for example.

### Initialization
- In order to allow a slate dialog to be used, the initialization of the VA system has been moved later in pre init, to the first avaliable point after slate has initialized and its shaders have been built. This in theory should not cause a a problem as engine content cannot be virtualized.
-- Added a new option 'engine.ini:[Core.ContentVirtualization]:InitPreSlate' with a default of false. When set to true the VA system will initialize in the original, earlier portion of engine pre init. This is only included in case an existing virtualized project encounters a problem with the new ordering.

diff --git a/Engine/Config/BaseEngine.ini b/Engine/Config/BaseEngine.ini
--- a/Engine/Config/BaseEngine.ini
+++ b/Engine/Config/BaseEngine.ini
@@ -1371,6 +1371,7 @@ AllowVirtualizationOptOut =False
 [Core.ContentVirtualization]
 SystemName=None
 LazyInit=False
+InitPreSlate=False
 
 [Core.VirtualizationModule]
 EnablePayloadVirtualization=True
diff --git a/Engine/Source/Developer/Virtualization/Private/SVirtualizationRevisionControlConnectionDialog.cpp b/Engine/Source/Developer/Virtualization/Private/SVirtualizationRevisionControlConnectionDialog.cpp
new file mode 100644
--- /dev/null
+++ b/Engine/Source/Developer/Virtualization/Private/SVirtualizationRevisionControlConnectionDialog.cpp
@@ -0,0 +1,298 @@
+// Copyright Epic Games, Inc. All Rights Reserved.
+
+#include "SVirtualizationRevisionControlConnectionDialog.h"
+
+#if UE_VA_WITH_SLATE
+
+#include "Async/ManualResetEvent.h"
+#include "Framework/Application/SlateApplication.h"
+#include "Framework/Docking/TabManager.h"
+#include "HAL/IConsoleManager.h"
+#include "Misc/App.h"
+#include "Misc/ConfigCacheIni.h"
+#include "SourceControlHelpers.h"
+#include "Widgets/Images/SImage.h"
+#include "Widgets/Input/SButton.h"
+#include "Widgets/Input/SEditableTextBox.h"
+#include "Widgets/SBoxPanel.h"
+#include "Widgets/Text/STextBlock.h"
+
+#define LOCTEXT_NAMESPACE "Virtualization"
+
+namespace UE::Virtualization
+{
+
+SRevisionControlConnectionDialog::FResult SRevisionControlConnectionDialog::RunDialog(FStringView CurrentPort, FStringView CurrentUsername)
+{
+	if (FApp::IsUnattended())
+	{
+		UE_LOG(LogVirtualization, Warning, TEXT("Skipping attempt to show SRevisionControlConnectionDialog as the application is unattended"));
+		return FResult();
+	}
+
+	if (!IsInGameThread())
+	{
+		UE_LOG(LogVirtualization, Warning, TEXT("Attempting to show SRevisionControlConnectionDialog on a worker thread!"));
+		return FResult();
+	}
+
+	if (!FSlateApplication::IsInitialized() || FSlateApplication::Get().GetRenderer() == nullptr)
+	{
+		UE_LOG(LogVirtualization, Warning, TEXT("Attempting to show SRevisionControlConnectionDialog before slate is initialized"));
+		return FResult();
+	}
+
+	UE_LOG(LogVirtualization, Display, TEXT("Creating dialog"));
+
+	TSharedPtr<SWindow> DialogWindow = SNew(SWindow)
+		.Title(LOCTEXT("VASCSettings", "Perforce Source Control Backend Settings"))
+		.FocusWhenFirstShown(true)
+		.SupportsMinimize(false)
+		.SupportsMaximize(false)
+		.SizingRule(ESizingRule::Autosized)
+		.HasCloseButton(false);
+
+	TSharedPtr<SRevisionControlConnectionDialog> DialogWidget;
+
+	TSharedPtr<SBorder> DialogWrapper =
+		SNew(SBorder)
+		.BorderImage(FAppStyle::GetBrush("ToolPanel.GroupBorder"))
+		.Padding(4.0f)
+		[
+			SNew(SVerticalBox)
+			+ SVerticalBox::Slot()
+			.AutoHeight()
+			.Padding(16.0f, 16.0f, 16.0f, 0.0f)
+			[
+				SAssignNew(DialogWidget, SRevisionControlConnectionDialog, CurrentPort, CurrentUsername)
+				.Window(DialogWindow)
+			]
+		];
+
+	DialogWindow->SetContent(DialogWrapper.ToSharedRef());
+
+	UE_LOG(LogVirtualization, Display, TEXT("Connection to source control for virtualized assets failed. Offering the user the choice to retry or continue anyway"));
+
+	TSharedPtr<SWindow> ParentWindow = FSlateApplication::Get().GetActiveTopLevelWindow();
+	FSlateApplication::Get().AddModalWindow(DialogWindow.ToSharedRef(), ParentWindow);
+
+	if (DialogWidget->GetResult() == SRevisionControlConnectionDialog::EResult::Retry)
+	{
+		return FResult(DialogWidget->GetPort(), DialogWidget->GetUserName());
+	}
+
+	return FResult();
+}
+
+void SRevisionControlConnectionDialog::Construct(const FArguments& InArgs, FStringView CurrentPort, FStringView CurrentUsername)
+{
+	WindowWidget = InArgs._Window;
+
+	const FString CurPort = TEXT("<P4PORT Here>");
+	const FString CurUser = TEXT("<P4USER Here>");
+
+	ChildSlot
+	[
+		SNew(SVerticalBox)
+		+ SVerticalBox::Slot()
+		.Padding(FMargin(0.0f, 0.0f, 16.0f, 16.0f))
+		.AutoHeight()
+		[
+			SNew(SHorizontalBox)
+			+ SHorizontalBox::Slot()
+			.Padding(FMargin(0.0f, 0.0f, 8.0f, 0.0f))
+			.HAlign(HAlign_Left)
+			.AutoWidth()
+			[
+				SNew(SVerticalBox)
+				+SVerticalBox::Slot()
+				.Padding(FMargin(0.0f, 0.0f, 0.0f, 0.0f))
+				.AutoHeight()
+				[
+					SNew(SImage)
+					.Image(FAppStyle::GetBrush("Icons.Error"))
+				]
+			]
+			+ SHorizontalBox::Slot()
+			[
+				SNew(STextBlock)
+				.Text(LOCTEXT("VASCMsg", "Failed to connect to the source control backend.\nThis may prevent you from accessing virtualized data in the future.\n\nPlease enter the correct source control settings below."))
+				.AutoWrapText(true)
+			]
+		]
+		+ SVerticalBox::Slot()
+		.Padding(FMargin(0.0f, 0.0f, 16.0f, 16.0f))
+		.AutoHeight()
+		[
+			SNew(SHorizontalBox)
+			+ SHorizontalBox::Slot()
+			.FillWidth(1.0f)
+			.Padding(FMargin(0.0f, 0.0f, 16.0f, 0.0f))
+			[
+				SNew(SVerticalBox)
+				+ SVerticalBox::Slot()
+				.Padding(FMargin(0.0f, 0.0f, 0.0f, 10.0f))
+				.VAlign(VAlign_Center)
+				.HAlign(HAlign_Right)
+				[
+					SNew(STextBlock)
+					.Text(LOCTEXT("PortLabel", "Server"))
+					.ToolTipText(LOCTEXT("PortLabel_Tooltip", "The server and port for your Perforce server. Usage ServerName:1234."))
+				]
+				+ SVerticalBox::Slot()
+				.Padding(FMargin(0.0f, 0.0f, 0.0f, 10.0f))
+				.VAlign(VAlign_Center)
+				.HAlign(HAlign_Right)
+				[
+					SNew(STextBlock)
+					.Text(LOCTEXT("UserNameLabel", "User Name"))
+					.ToolTipText(LOCTEXT("UserNameLabel_Tooltip", "Perforce username."))
+				]
+			]
+			+ SHorizontalBox::Slot()
+			.FillWidth(2.0f)
+			[
+				SNew(SVerticalBox)
+				+ SVerticalBox::Slot()
+				.VAlign(VAlign_Center)
+				.Padding(FMargin(0.0f, 0.0f, 0.0f, 10.0f))
+				[
+					SAssignNew(PortTextWidget, SEditableTextBox)
+					.Text(FText::FromString(FString(CurrentPort)))
+					.ToolTipText(LOCTEXT("VASC_PortTip", "The server and port for your Perforce server. Usage ServerName:1234."))
+				]
+				+ SVerticalBox::Slot()
+				.Padding(FMargin(0.0f, 0.0f, 0.0f, 10.0f))
+				.VAlign(VAlign_Center)
+				[
+					SAssignNew(UsernameTextWidget, SEditableTextBox)
+					.Text(FText::FromString(FString(CurrentUsername)))
+					.ToolTipText(LOCTEXT("VASC_UserTip", "Perforce username."))
+				]
+			]
+		]
+		+ SVerticalBox::Slot()
+		.Padding(FMargin(0.0f, 0.0f, 16.0f, 16.0f))
+		.AutoHeight()
+		[
+			SNew(SHorizontalBox)
+			+ SHorizontalBox::Slot()
+			.Padding(FMargin(0.0f, 0.0f, 8.0f, 0.0f))
+			.HAlign(HAlign_Left)
+			.AutoWidth()
+			[
+				SNew(SVerticalBox)
+				+ SVerticalBox::Slot()
+				.Padding(FMargin(0.0f, 0.0f, 0.0f, 0.0f))
+				.AutoHeight()
+				[
+					SNew(SImage)
+					.Image(FAppStyle::GetBrush("Icons.Warning"))
+				]
+			]
+			+ SHorizontalBox::Slot()
+			[
+				SNew(STextBlock)
+				.Text(LOCTEXT("VASCSkipWarning", "Skipping may cause future editor instability if virtualized data is required!"))
+				.AutoWrapText(true)
+			]
+		]
+		+ SVerticalBox::Slot()
+		.Padding(FMargin(0.0f, 0.0f, 16.0f, 16.0f))
+		.AutoHeight()
+		[
+			SNew(SHorizontalBox)
+			+ SHorizontalBox::Slot()
+			.HAlign(HAlign_Left)
+			[
+				SNew(SButton)
+				.VAlign(VAlign_Center)
+				.TextStyle(FAppStyle::Get(), "DialogButtonText")
+				.Text(LOCTEXT("VASC_Reset", "Reset To Defaults"))
+				.OnClicked(this, &SRevisionControlConnectionDialog::OnResetToDefaults)
+			]
+			+ SHorizontalBox::Slot()
+			.HAlign(HAlign_Right)
+			[
+				SNew(SHorizontalBox)
+				+ SHorizontalBox::Slot()
+				.AutoWidth()
+				.Padding(FMargin(5.0f, 0.0f))
+				[
+					SNew(SButton)
+					.VAlign(VAlign_Center)
+					.ButtonStyle(FAppStyle::Get(), "PrimaryButton")
+					.TextStyle(FAppStyle::Get(), "DialogButtonText")
+					.Text(LOCTEXT("VASC_Retry", "Retry Connection"))
+					.OnClicked(this, &SRevisionControlConnectionDialog::OnRetryConnection)
+				]
+				+ SHorizontalBox::Slot()
+				.AutoWidth()
+				.Padding(FMargin(5.0f, 0.0f))
+				[
+					SNew(SButton)
+					.VAlign(VAlign_Center)
+					.TextStyle(FAppStyle::Get(), "DialogButtonText")
+					.Text(LOCTEXT("VASC_Skip", "Skip"))
+					.OnClicked(this, &SRevisionControlConnectionDialog::OnSkip)
+				]
+			]
+		]
+	];
+}
+
+void SRevisionControlConnectionDialog::CloseModalDialog()
+{
+	if (WindowWidget.IsValid())
+	{
+		WindowWidget.Pin()->RequestDestroyWindow();
+	}
+}
+
+FReply SRevisionControlConnectionDialog::OnResetToDefaults()
+{
+	GConfig->EmptySection(TEXT("PerforceSourceControl.VirtualizationSettings"), SourceControlHelpers::GetSettingsIni());
+
+	Port.Empty();
+	UserName.Empty();
+
+	CloseModalDialog();
+
+	return FReply::Handled();
+}
+
+FReply SRevisionControlConnectionDialog::OnRetryConnection()
+{
+	UE_LOG(LogVirtualization, Display, TEXT("User opted to retry connecting to source control"));
+
+	Result = EResult::Retry;
+
+	Port = PortTextWidget->GetText().ToString();
+	UserName = UsernameTextWidget->GetText().ToString();
+
+	CloseModalDialog();
+
+	return FReply::Handled();
+}
+
+FReply SRevisionControlConnectionDialog::OnSkip()
+{
+	UE_LOG(LogVirtualization, Warning, TEXT("User opted not to connect to source control. Virtualized data may not be accessible!"));
+
+	Result = EResult::Skip;
+
+	CloseModalDialog();
+
+	return FReply::Handled();
+}
+
+static FAutoConsoleCommand CCmdTestDialog = FAutoConsoleCommand(
+	TEXT("TestVADialog"),
+	TEXT(""),
+	FConsoleCommandDelegate::CreateStatic(SRevisionControlConnectionDialog::RunDialogCvar));
+
+} // namespace UE::Virtualization
+
+#undef LOCTEXT_NAMESPACE
+
+#endif //UE_VA_WITH_SLATE
diff --git a/Engine/Source/Developer/Virtualization/Private/SVirtualizationRevisionControlConnectionDialog.h b/Engine/Source/Developer/Virtualization/Private/SVirtualizationRevisionControlConnectionDialog.h
new file mode 100644
--- /dev/null
+++ b/Engine/Source/Developer/Virtualization/Private/SVirtualizationRevisionControlConnectionDialog.h
@@ -0,0 +1,94 @@
+// Copyright Epic Games, Inc. All Rights Reserved.
+
+#pragma once
+
+#if UE_VA_WITH_SLATE
+
+#include "Input/Reply.h"
+#include "Widgets/SCompoundWidget.h"
+
+class SEditableTextBox;
+class SWindow;
+
+namespace UE::Virtualization
+{
+
+class SRevisionControlConnectionDialog : public SCompoundWidget
+{
+public:
+	struct FResult
+	{
+	public:
+		FResult() = default;
+		FResult(const FString& InPort, const FString& InUserName)
+			: bShouldRetry(true)
+			, Port(InPort)
+			, UserName(InUserName)
+		{
+
+		}
+
+		~FResult() = default;
+
+	public:
+		bool bShouldRetry = false;
+
+		FString Port;
+		FString UserName;
+	};
+
+	static FResult RunDialog(FStringView CurrentPort, FStringView CurrentUsername);
+	static void RunDialogCvar()
+	{
+		RunDialog(TEXT("<P4PORT Here>"), TEXT("<P4USER Here>"));
+	}
+
+	SLATE_BEGIN_ARGS(SRevisionControlConnectionDialog) {}
+		SLATE_ARGUMENT(TSharedPtr<SWindow>, Window)
+	SLATE_END_ARGS()
+
+	void Construct(const FArguments& InArgs, FStringView CurrentPort, FStringView CurrentUsername);
+
+	enum class EResult : uint8
+	{
+		Retry,
+		Skip
+	};
+
+	EResult GetResult() const
+	{
+		return Result;
+	}
+
+	const FString& GetPort() const
+	{
+		return Port;
+	}
+
+	const FString& GetUserName() const
+	{
+		return UserName;
+	}
+
+private:
+
+	void CloseModalDialog();
+
+	FReply OnResetToDefaults();
+	FReply OnRetryConnection();
+	FReply OnSkip();
+
+	TWeakPtr<SWindow> WindowWidget;
+
+	TSharedPtr<SEditableTextBox> PortTextWidget;
+	TSharedPtr<SEditableTextBox> UsernameTextWidget;
+
+	EResult Result = EResult::Skip;
+
+	FString Port;
+	FString UserName;
+};
+
+} // namespace UE::Virtualization
+
+#endif //UE_VA_WITH_SLATE
diff --git a/Engine/Source/Developer/Virtualization/Private/VirtualizationSourceControlBackend.cpp b/Engine/Source/Developer/Virtualization/Private/VirtualizationSourceControlBackend.cpp
--- a/Engine/Source/Developer/Virtualization/Private/VirtualizationSourceControlBackend.cpp
+++ b/Engine/Source/Developer/Virtualization/Private/VirtualizationSourceControlBackend.cpp
@@ -15,6 +15,8 @@
 #include "Misc/PathViews.h"
 #include "Misc/Paths.h"
 #include "Misc/ScopeExit.h"
+#include "SVirtualizationRevisionControlConnectionDialog.h"
+#include "SourceControlHelpers.h"
 #include "SourceControlInitSettings.h"
 #include "SourceControlOperations.h"
 #include "VirtualizationManager.h"
@@ -336,17 +338,81 @@ IVirtualizationBackend::EConnectionStatus FSourceControlBackend::OnConnect()
 		return IVirtualizationBackend::EConnectionStatus::Error;
 	}
 
+	FString Port = ServerAddress;
+	FString UserName = TEXT("");
+	bool bSaveSettings = false; // Initially we will try to reason the perforce settings from the ini file and global environment
+								// so we don't want to save these values to the users ini files.
+
+	while (true)
+	{
+		// TODO: At the moment we cannot get back error messages from ISourceControlProvider::Init, nor can we get back info about
+		// the actual login details used, so for the initial dialog we will just have to go with the settings we know and not 
+		// give the user any details about the connection problem.
+		FText ErrorMessage;
+		IVirtualizationBackend::EConnectionStatus ConnectionResult = OnConnectInternal(Port, UserName, bSaveSettings, ErrorMessage);
+
+		if (ConnectionResult == IVirtualizationBackend::EConnectionStatus::Connected)
+		{
+			if (bSaveSettings)
+			{
+				check(::IsInGameThread());
+				GConfig->Flush(false, SourceControlHelpers::GetSettingsIni());
+			}
+			return IVirtualizationBackend::EConnectionStatus::Connected;
+		}
+
+#if UE_VA_WITH_SLATE
+		// If the local ini settings are ignored then there is no point saving correct settings given by the user.
+		// They will need to fix their root problem instead.
+		// TODO: Maybe give a bespoke error at this point?
+		if (bUseRetryConnectionDialog && bUseLocalIniFileSettings)
+		{
+			SRevisionControlConnectionDialog::FResult DialogResult = SRevisionControlConnectionDialog::RunDialog(Port, UserName);
+			if (!DialogResult.bShouldRetry)
+			{
+				OnConnectionError(ErrorMessage);
+				return IVirtualizationBackend::EConnectionStatus::Error;
+			}
+
+			Port = DialogResult.Port;
+			UserName = DialogResult.UserName;
+			bSaveSettings = true; // Now we have input from the user, so we should save it for future sessions
+		}
+		else
+#endif
+		{
+			OnConnectionError(ErrorMessage);
+			return IVirtualizationBackend::EConnectionStatus::Error;
+		}
+	}
+}
+
+
+IVirtualizationBackend::EConnectionStatus FSourceControlBackend::OnConnectInternal(FStringView Port, FStringView Username, bool bSaveConnectionSettings, FText& OutErrorMessage)
+{
 	// We do not want the connection to have a client workspace so explicitly set it to empty
 	FSourceControlInitSettings SCCSettings(FSourceControlInitSettings::EBehavior::OverrideExisting);
 
-	const FSourceControlInitSettings::EConfigBehavior IniBehavior = bUseLocalIniFileSettings ?	FSourceControlInitSettings::EConfigBehavior::ReadOnly :
-																								FSourceControlInitSettings::EConfigBehavior::None;
+	FSourceControlInitSettings::EConfigBehavior IniBehavior = bSaveConnectionSettings	? FSourceControlInitSettings::EConfigBehavior::ReadWrite 
+																						: FSourceControlInitSettings::EConfigBehavior::ReadOnly;
+
+	if (!bUseLocalIniFileSettings)
+	{
+		IniBehavior = FSourceControlInitSettings::EConfigBehavior::None;
+	}
+
 	SCCSettings.SetConfigBehavior(IniBehavior);
 
-	if (!ServerAddress.IsEmpty())
+	if (!Port.IsEmpty())
 	{
-		SCCSettings.AddSetting(TEXT("P4Port"), ServerAddress);
+		SCCSettings.AddSetting(TEXT("P4Port"), Port);
 	}
+
+	if (!Username.IsEmpty())
+	{
+		SCCSettings.AddSetting(TEXT("P4User"), Username);
+	}
+
 	SCCSettings.AddSetting(TEXT("P4Client"), TEXT(""));
 
 	SCCProvider = ISourceControlModule::Get().CreateProvider(FName("Perforce"), TEXT("Virtualization"), SCCSettings);
@@ -358,12 +424,12 @@ IVirtualizationBackend::EConnectionStatus FSourceControlBackend::OnConnect()
 
 	// Will attempt to make initial connection to the server, it doesn't return any error values we can check against but
 	// it will output problems to the log file.
-	SCCProvider->Init(true);
+	const bool bForceConnection = true;
+	SCCProvider->Init(bForceConnection);
 
 	if (!SCCProvider->IsAvailable())
 	{
-		OnConnectionError(FText(LOCTEXT("FailedSourceControlConnection", "Failed to connect to revision control backend, see the Message Log 'Revision Control' errors for details.\nThe revision control backend will be unable to pull payloads, is your revision control config set up correctly?")));
-
+		OutErrorMessage = FText(LOCTEXT("FailedSourceControlConnection", "Failed to connect to revision control backend, see the Message Log 'Revision Control' errors for details.\nThe revision control backend will be unable to pull payloads, is your revision control config set up correctly?"));
 		return IVirtualizationBackend::EConnectionStatus::Error;
 	}
 
@@ -388,12 +454,10 @@ IVirtualizationBackend::EConnectionStatus FSourceControlBackend::OnConnect()
 
 	if (MetaInfoBuffer.IsNull())
 	{
-		FText Msg = FText::Format(LOCTEXT("FailedMetaInfo", "Failed to find '{0}' on server '{1}'\nThe revision control backend will be unable to pull payloads, is your revision control config set up correctly?"),
+		OutErrorMessage = FText::Format(LOCTEXT("FailedMetaInfo", "Failed to find '{0}' on server '{1}'\nThe revision control backend will be unable to pull payloads, is your revision control config set up correctly?"),
 			FText::FromString(PayloadMetaInfoPath),
 			FText::FromString(GetPortFromProvider(SCCProvider.Get())));
 
-		OnConnectionError(MoveTemp(Msg));
-
 		return IVirtualizationBackend::EConnectionStatus::Error;
 	}
 
@@ -1009,6 +1073,11 @@ bool FSourceControlBackend::TryApplySettingsFromConfigFiles(const FString& Confi
 		UE_LOG(LogVirtualization, Log, TEXT("[%s] Reading settings from local SourceControlSettings.ini is %s"), *GetDebugName(), bUseLocalIniFileSettings ? TEXT("enabled") : TEXT("disabled"));	
 	}
 
+	{
+		FParse::Bool(*ConfigEntry, TEXT("UseRetryConnectionDialog="), bUseRetryConnectionDialog);
+		UE_LOG(LogVirtualization, Log, TEXT("[%s] Showing a reconnect dialog on initial failure %s"), *GetDebugName(), bUseRetryConnectionDialog ? TEXT("enabled") : TEXT("disabled"));
+	}
+
 	{
 		// TODO: We should just extract this from the perforce environment but that requires extending
 		// the source control api.
diff --git a/Engine/Source/Developer/Virtualization/Private/VirtualizationSourceControlBackend.h b/Engine/Source/Developer/Virtualization/Private/VirtualizationSourceControlBackend.h
--- a/Engine/Source/Developer/Virtualization/Private/VirtualizationSourceControlBackend.h
+++ b/Engine/Source/Developer/Virtualization/Private/VirtualizationSourceControlBackend.h
@@ -27,10 +27,10 @@ class FSemaphore;
  * 
  * Required Values:
  * 
- * DepotPath [string]			The path (in depot format) to the location where virtualized payloads are stored.
- *								[Default=""]
- * ClientStream [string]:		Used when the payloads are stored in a stream based depot. It should contain
- *								the stream name to use when creating a workspace for payload submission. [Default=""]
+ * DepotPath [string]				The path (in depot format) to the location where virtualized payloads are stored.
+ *									[Default=""]
+ * ClientStream [string]:			Used when the payloads are stored in a stream based depot. It should contain
+ *									the stream name to use when creating a workspace for payload submission. [Default=""]
  * 
  * If the payloads are being stored in a depot with type 'stream' them config value 'ClientStream' must be set to a valid
  * stream. The 'DepotPath' value can still be set to choose a specific location within that stream or when the stream 
@@ -40,41 +40,44 @@ class FSemaphore;
  * 
  * Optional Values:
  * 
- * Server [string]:				When set the backend will use this server address to connect to. When not
- *								set the backend will use the environment defaults. [Default=""]
- * UsePartitionedClient [bool]:	When true the temporary workspace client created to submit payloads 
- *								from will be created as a partitioned workspace which is less overhead
- *								on the source control server. If your server does not support this then
- *								use false. [Default=True]
- * WorkingDir [string]:			Sets the temp location on disk where payloads are submitted from. The path can
- *								contain config file expansions (see ConfigCacheIni.cpp ::MatchExpansions) and
- *								environment variables with the format $(EnvVarName).
- *								[Default="%GAMESAVEDDIR%/VASubmission"]
- * RetryCount [int32]:			How many times we should try to download a payload before giving up with
- *								an error. Useful when the connection is unreliable but does not experience 
- *								frequent persistent outages. [Default=2]
- * RetryWaitTime [int32]:		The length of time the process should wait between each download attempt
- *								in milliseconds. Remember that the max length of time that the process
- *								can stall attempting to download a payload file is 
- *								RetryCount * RetryWaitTime; [Default=100ms]
- * MaxConnections [int32]		The maximum number of perforce connections that can be made concurrently. 
- *								If this value is exceeded then additional requests will wait until a 
- *								connection becomes free. Pass in '-1' to disable this feature and allow
- *								unlimited connections. [Default=8]
- * BatchCount [int32]			The max number of payloads that can be pushed to source control in a
- *								single submit. If the number of payloads in a request batch exceeds
- *								this size then it will be split into multiple smaller batches. [Default=100]
- * SuppressNotifications[bool]:	When true the system will not display a pop up notification when a 
- *								connection error occurs, allowing the user to stay unaware of the error
- *								unless it actually causes some sort of problem. [Default=false]
- * UseLocalIniFileSettings[bool]When true the revision control provider will be allowed to load connection 
- *								settings from the users locally saved 'SourceControlSettings.ini' file, if 
- *								false then the settings in this file (if any) will be ignored. [Default=true]
- * IgnoreFile [string]:			Sets the name of the p4 ignore file to use. When submitting payloads we 
- *								create a custom p4 ignore file to override any ignore settings for a project
- *								which allows us to submit from the saved directory which is normally prevented
- *								by the default ignore file. This value can be set to what ever is used by
- *								your perforce environment. [Default=".p4ignore.txt"]
+ * Server [string]:					When set the backend will use this server address to connect to. When not
+ *									set the backend will use the environment defaults. [Default=""]
+ * UsePartitionedClient [bool]:		When true the temporary workspace client created to submit payloads 
+ *									from will be created as a partitioned workspace which is less overhead
+ *									on the source control server. If your server does not support this then
+ *									use false. [Default=True]
+ * WorkingDir [string]:				Sets the temp location on disk where payloads are submitted from. The path can
+ *									contain config file expansions (see ConfigCacheIni.cpp ::MatchExpansions) and
+ *									environment variables with the format $(EnvVarName).
+ *									[Default="%GAMESAVEDDIR%/VASubmission"]
+ * RetryCount [int32]:				How many times we should try to download a payload before giving up with
+ *									an error. Useful when the connection is unreliable but does not experience 
+ *									frequent persistent outages. [Default=2]
+ * RetryWaitTime [int32]:			The length of time the process should wait between each download attempt
+ *									in milliseconds. Remember that the max length of time that the process
+ *									can stall attempting to download a payload file is 
+ *									RetryCount * RetryWaitTime; [Default=100ms]
+ * MaxConnections [int32]			The maximum number of perforce connections that can be made concurrently. 
+ *									If this value is exceeded then additional requests will wait until a 
+ *									connection becomes free. Pass in '-1' to disable this feature and allow
+ *									unlimited connections. [Default=8]
+ * BatchCount [int32]				The max number of payloads that can be pushed to source control in a
+ *									single submit. If the number of payloads in a request batch exceeds
+ *									this size then it will be split into multiple smaller batches. [Default=100]
+ * SuppressNotifications[bool]:		When true the system will not display a pop up notification when a 
+ *									connection error occurs, allowing the user to stay unaware of the error
+ *									unless it actually causes some sort of problem. [Default=false]
+ * UseLocalIniFileSettings[bool]	When true the revision control provider will be allowed to load connection 
+ *									settings from the users locally saved 'SourceControlSettings.ini' file, if 
+ *									false then the settings in this file (if any) will be ignored. [Default=true]
+ * IgnoreFile [string]:				Sets the name of the p4 ignore file to use. When submitting payloads we 
+ *									create a custom p4 ignore file to override any ignore settings for a project
+ *									which allows us to submit from the saved directory which is normally prevented
+ *									by the default ignore file. This value can be set to what ever is used by
+ *									your perforce environment. [Default=".p4ignore.txt"]
+ * UseRetryConnectionDialog[bool]	When true a slate dialog will be shown if the initial connection to the 
+ *									source control server fails allowing the user to attempt to input the correct
+ *									login values. [Default=false]&
  * 
  * Environment Variables:
  * UE-VirtualizationWorkingDir [string]:	This can be set to a valid directory path that the backend
@@ -126,6 +129,7 @@ private:
 	virtual bool Initialize(const FString& ConfigEntry) override;
 
 	virtual EConnectionStatus OnConnect() override;
+	IVirtualizationBackend::EConnectionStatus OnConnectInternal(FStringView Port, FStringView Username, bool bSaveConnectionSettings, FText& OutErrorMessage);
 	
 	virtual bool PushData(TArrayView<FPushRequest> Requests, EPushFlags Flags) override;
 	virtual bool PullData(TArrayView<FPullRequest> Requests, EPullFlags Flags, FText& OutErrors) override;
@@ -186,6 +190,9 @@ private:
 
 	/** When true we allow the revision control provider to read settings from the users local SourceControlSettings.ini file */
 	bool bUseLocalIniFileSettings = true;
+
+	/* When true, if the initial connection to the revision control provider fails we will show a slate dialog (if possible) to the user prompting for correct settings */
+	bool bUseRetryConnectionDialog = false;
 };
 
 } // namespace UE::Virtualization
diff --git a/Engine/Source/Developer/Virtualization/Virtualization.Build.cs b/Engine/Source/Developer/Virtualization/Virtualization.Build.cs
--- a/Engine/Source/Developer/Virtualization/Virtualization.Build.cs
+++ b/Engine/Source/Developer/Virtualization/Virtualization.Build.cs
@@ -21,5 +21,18 @@ public class Virtualization : ModuleRules
 				"Projects",
 				"SourceControl"
 			});
+
+		if (Target.bUsesSlate)
+		{
+			// NOTE: Slate is being included via SourceControl anyway
+			PrivateDependencyModuleNames.Add("Slate");
+			PrivateDependencyModuleNames.Add("SlateCore");
+
+			PublicDefinitions.Add("UE_VA_WITH_SLATE=1");
+		}
+		else
+		{
+			PublicDefinitions.Add("UE_VA_WITH_SLATE=0");
+		}
 	}
 }
diff --git a/Engine/Source/Runtime/Core/Private/Virtualization/VirtualizationSystem.cpp b/Engine/Source/Runtime/Core/Private/Virtualization/VirtualizationSystem.cpp
--- a/Engine/Source/Runtime/Core/Private/Virtualization/VirtualizationSystem.cpp
+++ b/Engine/Source/Runtime/Core/Private/Virtualization/VirtualizationSystem.cpp
@@ -349,6 +349,14 @@ void Initialize(const FInitParams& InitParams, EInitializationFlags Flags)
 	}
 }
 
+bool ShouldInitializePreSlate()
+{
+	bool bInitBeforeSlate = false;
+	GConfig->GetBool(TEXT("Core.ContentVirtualization"), TEXT("InitPreSlate"), bInitBeforeSlate, GEngineIni);
+
+	return bInitBeforeSlate;
+}
+
 void Shutdown()
 {
 	GVirtualizationSystem.Reset();
diff --git a/Engine/Source/Runtime/Core/Public/Virtualization/VirtualizationSystem.h b/Engine/Source/Runtime/Core/Public/Virtualization/VirtualizationSystem.h
--- a/Engine/Source/Runtime/Core/Public/Virtualization/VirtualizationSystem.h
+++ b/Engine/Source/Runtime/Core/Public/Virtualization/VirtualizationSystem.h
@@ -505,6 +505,24 @@ CORE_API void Initialize(EInitializationFlags Flags);
  */
 CORE_API void Initialize(const FInitParams& InitParams, EInitializationFlags Flags);
 
+/**
+ * Returns if the system should be initialized before the slate system has been initialized or not.
+ * 
+ * Originally the VA system would be initialized as early as possible, but this would occur before
+ * slate has been set up and prevents us from initiating interactive dialogs. Now that we have
+ * the firm rule that engine content cannot be virtualized it should be safe for us to initialize
+ * the system after slate allowing dialogs.
+ * 
+ * This function exists so that a project could opt into the original behavior by setting the 
+ * config file value "engine.ini:[Core.ContentVirtualization]:InitPreSlate=true"
+ * This is only provided as a way for teams to quickly fix any issue that the later initialization
+ * might cause and is not a long term feature. This will be deprecated and removed in future releases.
+ * 
+ * Note that this is only really supported in the editor as we use the global config file, not one
+ * provided. Standalone programs should just initialize the system as needed.
+ */
+CORE_API bool ShouldInitializePreSlate();
+
 /**
  * Shutdowns the global IVirtualizationSystem if it exists. 
  * Calling this is optional as the system will shut itself down along with the rest of the engine.
diff --git a/Engine/Source/Runtime/Launch/Private/LaunchEngineLoop.cpp b/Engine/Source/Runtime/Launch/Private/LaunchEngineLoop.cpp
--- a/Engine/Source/Runtime/Launch/Private/LaunchEngineLoop.cpp
+++ b/Engine/Source/Runtime/Launch/Private/LaunchEngineLoop.cpp
@@ -3244,10 +3244,12 @@ int32 FEngineLoop::PreInitPreStartupScreen(const TCHAR* CmdLine)
 		}
 #endif
 
-#if WITH_EDITOR
-		// The virtualization system will be initialized when ever IVirtualizationSystem::Get is called
-		// but here we try to explicitly initialize it.
-		UE::Virtualization::Initialize(UE::Virtualization::EInitializationFlags::None);
+#if WITH_EDITOR	
+		if (UE::Virtualization::ShouldInitializePreSlate())
+		{
+			// Explicit initialization of the virtualization system, before slate has initialized and we cannot show error dialogs 
+			UE::Virtualization::Initialize(UE::Virtualization::EInitializationFlags::None);
+		}
 #endif //WITH_EDITOR
 
 		check(!GDistanceFieldAsyncQueue);
@@ -3736,6 +3738,14 @@ int32 FEngineLoop::PreInitPostStartupScreen(const TCHAR* CmdLine)
 			ProcessNewlyLoadedUObjects();
 		}
 
+#if WITH_EDITOR
+		if (!UE::Virtualization::ShouldInitializePreSlate())
+		{
+			// Explicit initialization of the virtualization system, after slate has initialized and we can show error dialogs.
+			UE::Virtualization::Initialize(UE::Virtualization::EInitializationFlags::None);
+		}
+#endif //WITH_EDITOR
+
 		// Ensure game localization has loaded before we continue
 		FTextLocalizationManager::Get().WaitForAsyncTasks();
 
