Fix warning on editor startup introduced in CL 26115183

#rnx

- Any new option added to engine.ini:[Core.ContentVirtualization] needs to also be added to an allowed list in FVirtualizationManager so that we do not flag it as a deprecated option. This is because we used to store the manager config values in that section as well.

diff --git a/Engine/Source/Developer/Virtualization/Private/VirtualizationManager.cpp b/Engine/Source/Developer/Virtualization/Private/VirtualizationManager.cpp
--- a/Engine/Source/Developer/Virtualization/Private/VirtualizationManager.cpp
+++ b/Engine/Source/Developer/Virtualization/Private/VirtualizationManager.cpp
@@ -2,24 +2,23 @@
 
 #include "VirtualizationManager.h"
 
+#include "AnalyticsEventAttribute.h"
 #include "HAL/IConsoleManager.h"
 #include "HAL/PlatformTime.h"
 #include "Logging/MessageLog.h"
 #include "Misc/CommandLine.h"
 #include "Misc/ConfigCacheIni.h"
+#include "Misc/MessageDialog.h"
 #include "Misc/PackageName.h"
 #include "Misc/PackagePath.h"
 #include "Misc/Parse.h"
 #include "Misc/Paths.h"
 #include "Misc/ScopeLock.h"
-#include "PackageRehydrationProcess.h"
 #include "Misc/ScopedSlowTask.h"
+#include "PackageRehydrationProcess.h"
 #include "PackageVirtualizationProcess.h"
 #include "ProfilingDebugging/CookStats.h"
 #include "VirtualizationFilterSettings.h"
-#include "AnalyticsEventAttribute.h"
-
-#include "Misc/MessageDialog.h"
 
 #define LOCTEXT_NAMESPACE "Virtualization"
 
@@ -1165,8 +1164,8 @@ void FVirtualizationManager::ApplySettingsFromConfigFiles(const FConfigFile& Con
 
 	// Check for any legacy settings and print them out (easier to do this in one block rather than one and time)
 	{
-		// Entries that are allows to be in [Core.ContentVirtualization
-		static const TArray<FString> AllowedEntries = { TEXT("SystemName") , TEXT("LazyInit") };
+		// Entries that are allows to be in [Core.ContentVirtualization]
+		static const TArray<FString> AllowedEntries = { TEXT("SystemName") , TEXT("LazyInit"), TEXT("InitPreSlate") };
 		
 		TArray<FString> LegacyEntries;	
 		if (const FConfigSection* LegacySection = ConfigFile.Find(LegacyConfigSection))
