Improve VA revision control dialog by adding better tool tips and support for turning the 'ConnectionHelpUrl' (if set for a project) into a clickable link for the user.

#rnx

diff --git a/Engine/Source/Developer/Virtualization/Private/SVirtualizationRevisionControlConnectionDialog.cpp b/Engine/Source/Developer/Virtualization/Private/SVirtualizationRevisionControlConnectionDialog.cpp
--- a/Engine/Source/Developer/Virtualization/Private/SVirtualizationRevisionControlConnectionDialog.cpp
+++ b/Engine/Source/Developer/Virtualization/Private/SVirtualizationRevisionControlConnectionDialog.cpp
@@ -14,8 +14,10 @@
 #include "Widgets/Images/SImage.h"
 #include "Widgets/Input/SButton.h"
 #include "Widgets/Input/SEditableTextBox.h"
+#include "Widgets/Input/SHyperlink.h"
 #include "Widgets/SBoxPanel.h"
 #include "Widgets/Text/STextBlock.h"
+#include "VirtualizationManager.h"
 
 #define LOCTEXT_NAMESPACE "Virtualization"
 
@@ -91,6 +93,8 @@ void SRevisionControlConnectionDialog::Construct(const FArguments& InArgs, FStri
 	const FString CurPort = TEXT("<P4PORT Here>");
 	const FString CurUser = TEXT("<P4USER Here>");
 
+	const FString ConnectionHelpUrl = FVirtualizationManager::GetConnectionHelpUrl();
+
 	ChildSlot
 	[
 		SNew(SVerticalBox)
@@ -116,7 +120,7 @@ void SRevisionControlConnectionDialog::Construct(const FArguments& InArgs, FStri
 			+ SHorizontalBox::Slot()
 			[
 				SNew(STextBlock)
-				.Text(LOCTEXT("VASCMsg", "Failed to connect to the source control backend.\nThis may prevent you from accessing virtualized data in the future.\n\nPlease enter the correct source control settings below."))
+				.Text(LOCTEXT("VASCMsg", "Failed to connect to the source control backend.\nThis may prevent you from accessing virtualized data in the future.\n\nPlease enter the correct source control settings below:"))
 				.AutoWrapText(true)
 			]
 		]
@@ -198,7 +202,17 @@ void SRevisionControlConnectionDialog::Construct(const FArguments& InArgs, FStri
 			]
 		]
 		+ SVerticalBox::Slot()
-		.Padding(FMargin(0.0f, 0.0f, 16.0f, 16.0f))
+		.HAlign(HAlign_Right)
+		.AutoHeight()
+		[
+			SNew(SHyperlink)
+			.Text(LOCTEXT("VASCHelpUrl", "Click here for additional documentation"))
+			.ToolTipText(FText::FromString(ConnectionHelpUrl))
+			.OnNavigate(this, &SRevisionControlConnectionDialog::OnUrlClicked)
+			.Visibility_Lambda([ConnectionHelpUrl] { return !ConnectionHelpUrl.IsEmpty() ? EVisibility::Visible : EVisibility::Hidden; })
+		]
+		+ SVerticalBox::Slot()
+		.Padding(FMargin(0.0f, 16.0f, 16.0f, 16.0f))
 		.AutoHeight()
 		[
 			SNew(SHorizontalBox)
@@ -209,6 +223,7 @@ void SRevisionControlConnectionDialog::Construct(const FArguments& InArgs, FStri
 				.VAlign(VAlign_Center)
 				.TextStyle(FAppStyle::Get(), "DialogButtonText")
 				.Text(LOCTEXT("VASC_Reset", "Reset To Defaults"))
+				.ToolTipText(LOCTEXT("VASC_ResetTip", "Removes connection settings that may be saved to your local ini files and attempts to connect using your environment defaults"))
 				.OnClicked(this, &SRevisionControlConnectionDialog::OnResetToDefaults)
 			]
 			+ SHorizontalBox::Slot()
@@ -224,6 +239,7 @@ void SRevisionControlConnectionDialog::Construct(const FArguments& InArgs, FStri
 					.ButtonStyle(FAppStyle::Get(), "PrimaryButton")
 					.TextStyle(FAppStyle::Get(), "DialogButtonText")
 					.Text(LOCTEXT("VASC_Retry", "Retry Connection"))
+					.ToolTipText(LOCTEXT("VASC_RetryTip", "Attempts to reconnect to the revision control server with the settings that you entered"))
 					.OnClicked(this, &SRevisionControlConnectionDialog::OnRetryConnection)
 				]
 				+ SHorizontalBox::Slot()
@@ -234,6 +250,7 @@ void SRevisionControlConnectionDialog::Construct(const FArguments& InArgs, FStri
 					.VAlign(VAlign_Center)
 					.TextStyle(FAppStyle::Get(), "DialogButtonText")
 					.Text(LOCTEXT("VASC_Skip", "Skip"))
+					.ToolTipText(LOCTEXT("VASC_RetryTip", "The editor will continue to load but will be unable to pull virtualized data from revision control if needed"))
 					.OnClicked(this, &SRevisionControlConnectionDialog::OnSkip)
 				]
 			]
@@ -286,6 +303,13 @@ FReply SRevisionControlConnectionDialog::OnSkip()
 	return FReply::Handled();
 }
 
+void SRevisionControlConnectionDialog::OnUrlClicked() const
+{
+	const FString ConnectionHelpUrl = FVirtualizationManager::GetConnectionHelpUrl();
+
+	FPlatformProcess::LaunchURL(*ConnectionHelpUrl, nullptr, nullptr);
+}
+
 static FAutoConsoleCommand CCmdTestDialog = FAutoConsoleCommand(
 	TEXT("TestVADialog"),
 	TEXT(""),
diff --git a/Engine/Source/Developer/Virtualization/Private/SVirtualizationRevisionControlConnectionDialog.h b/Engine/Source/Developer/Virtualization/Private/SVirtualizationRevisionControlConnectionDialog.h
--- a/Engine/Source/Developer/Virtualization/Private/SVirtualizationRevisionControlConnectionDialog.h
+++ b/Engine/Source/Developer/Virtualization/Private/SVirtualizationRevisionControlConnectionDialog.h
@@ -78,6 +78,8 @@ private:
 	FReply OnRetryConnection();
 	FReply OnSkip();
 
+	void OnUrlClicked() const;
+
 	TWeakPtr<SWindow> WindowWidget;
 
 	TSharedPtr<SEditableTextBox> PortTextWidget;
