Add a command line '-VADisableDialog' that will prevent the revision control retry dialog from being disabled as an easy way for users to supress it if they need to do so.

#rnx

diff --git a/Engine/Source/Developer/Virtualization/Private/VirtualizationSourceControlBackend.cpp b/Engine/Source/Developer/Virtualization/Private/VirtualizationSourceControlBackend.cpp
--- a/Engine/Source/Developer/Virtualization/Private/VirtualizationSourceControlBackend.cpp
+++ b/Engine/Source/Developer/Virtualization/Private/VirtualizationSourceControlBackend.cpp
@@ -10,6 +10,7 @@
 #include "ISourceControlProvider.h"
 #include "Interfaces/IPluginManager.h"
 #include "Logging/MessageLog.h"
+#include "Misc/CommandLine.h"
 #include "Misc/FileHelper.h"
 #include "Misc/Parse.h"
 #include "Misc/PathViews.h"
@@ -1075,6 +1076,12 @@ bool FSourceControlBackend::TryApplySettingsFromConfigFiles(const FString& Confi
 
 	{
 		FParse::Bool(*ConfigEntry, TEXT("UseRetryConnectionDialog="), bUseRetryConnectionDialog);
+
+		if (FParse::Param(FCommandLine::Get(), TEXT("VADisableDialog")))
+		{
+			bUseRetryConnectionDialog = false;
+		}
+
 		UE_LOG(LogVirtualization, Log, TEXT("[%s] Showing a reconnect dialog on initial failure %s"), *GetDebugName(), bUseRetryConnectionDialog ? TEXT("enabled") : TEXT("disabled"));
 	}
 
