Fix incorrect localization key

#rnx

diff --git a/Engine/Source/Developer/Virtualization/Private/SVirtualizationRevisionControlConnectionDialog.cpp b/Engine/Source/Developer/Virtualization/Private/SVirtualizationRevisionControlConnectionDialog.cpp
--- a/Engine/Source/Developer/Virtualization/Private/SVirtualizationRevisionControlConnectionDialog.cpp
+++ b/Engine/Source/Developer/Virtualization/Private/SVirtualizationRevisionControlConnectionDialog.cpp
@@ -250,7 +250,7 @@ void SRevisionControlConnectionDialog::Construct(const FArguments& InArgs, FStri
 					.VAlign(VAlign_Center)
 					.TextStyle(FAppStyle::Get(), "DialogButtonText")
 					.Text(LOCTEXT("VASC_Skip", "Skip"))
-					.ToolTipText(LOCTEXT("VASC_RetryTip", "The editor will continue to load but will be unable to pull virtualized data from revision control if needed"))
+					.ToolTipText(LOCTEXT("VASC_SkipTip", "The editor will continue to load but will be unable to pull virtualized data from revision control if needed"))
 					.OnClicked(this, &SRevisionControlConnectionDialog::OnSkip)
 				]
 			]
