Add a new config option 'engine.ini:[Core.VirtualizationModule]:DisableLazyInitIfInteractive' that can be used to override LazyInitConnections if we want the VA module to be able to use slate dialogs safely.

- If we want to use slate dialogs we cannot use LazyInitConnections, however slate dialogs only work on programs that a) support slate and b) are running in some interactive form. Commandline programs, or the editor runnning in unattended mode or running a commandlet cannot make use of the slate dialog.
- If this new flag 'DisableLazyInitIfInteractive' is set to true and 'LazyInitConnections' is set to true then the VA system will lazy init its connections unless we are running an interactive program that could make use of the slate dialogs, in which case 'LazyInitConnections' is disabled for safety.
- Ideally we'd merge these two values into a single enum but I am hoping that we can change how we connect in the near term in which case we will simple deprecate 'DisableLazyInitIfInteractive'. The use case is pretty niche at the moment so this will probably be easier for the next stage of work.

diff --git a/Engine/Config/BaseEngine.ini b/Engine/Config/BaseEngine.ini
--- a/Engine/Config/BaseEngine.ini
+++ b/Engine/Config/BaseEngine.ini
@@ -1384,6 +1384,7 @@ AllowSubmitIfVirtualizationFailed=false
 FilterMode=OptOut
 FilterMapContent=True
 LazyInitConnections=False
+DisableLazyInitIfInteractive=False
 UseLegacyErrorHandling=False
 ForceCachingOnPull=False
 PullErrorAdditionalMsg=""
diff --git a/Engine/Source/Developer/Virtualization/Private/VirtualizationManager.cpp b/Engine/Source/Developer/Virtualization/Private/VirtualizationManager.cpp
--- a/Engine/Source/Developer/Virtualization/Private/VirtualizationManager.cpp
+++ b/Engine/Source/Developer/Virtualization/Private/VirtualizationManager.cpp
@@ -6,6 +6,7 @@
 #include "HAL/IConsoleManager.h"
 #include "HAL/PlatformTime.h"
 #include "Logging/MessageLog.h"
+#include "Misc/App.h"
 #include "Misc/CommandLine.h"
 #include "Misc/ConfigCacheIni.h"
 #include "Misc/MessageDialog.h"
@@ -1120,6 +1121,25 @@ void FVirtualizationManager::ApplySettingsFromConfigFiles(const FConfigFile& Con
 
 #if UE_VIRTUALIZATION_CONNECTION_LAZY_INIT == 0
 	ConfigFile.GetBool(ConfigSection, TEXT("LazyInitConnections"), bLazyInitConnections);
+
+	if (bLazyInitConnections)
+	{
+		// Check if we should override the value of bLazyInitConnections depending on if the current process is interactive or not
+		bool bDisableLazyInitIfInteractive = false;
+		ConfigFile.GetBool(ConfigSection, TEXT("DisableLazyInitIfInteractive"), bDisableLazyInitIfInteractive);
+
+		const bool bIsUsingSlate = UE_VA_WITH_SLATE != 0;
+		const bool bIsUnattended = FApp::IsUnattended();
+		const bool bIsRunningCommandlet = IsRunningCommandlet();
+
+		const bool bIsInteractive = bIsUsingSlate && !bIsUnattended && !bIsRunningCommandlet;
+
+		if (bDisableLazyInitIfInteractive && bIsInteractive)
+		{
+			bLazyInitConnections = false;
+		}
+	}
+
 	UE_LOG(LogVirtualization, Display, TEXT("\tLazyInitConnections : %s"), bLazyInitConnections ? TEXT("true") : TEXT("false"));
 #else
 	bLazyInitConnections = true;
diff --git a/Engine/Source/Developer/Virtualization/Private/VirtualizationManager.h b/Engine/Source/Developer/Virtualization/Private/VirtualizationManager.h
--- a/Engine/Source/Developer/Virtualization/Private/VirtualizationManager.h
+++ b/Engine/Source/Developer/Virtualization/Private/VirtualizationManager.h
@@ -114,6 +114,14 @@ struct FAnalyticsEventAttribute;
  *												This can remove lengthy connection steps from the process init phase and then only connect
  *												if we actually need that service. Note that if this is true then the connection can come from
  *												any thread, so custom backend code will need to take that into account. [Default=false]
+ * DisableLazyInitIfInteractive [bool]			When true 'LazyInitConnections' will be forced to false if slate is enabled. This exists because
+ *												some backends can show slate dialogs when their initial connection fails to prompt for the 
+ *												correct login values. When 'LazyInitConnections' is true, this request can come on any thread and
+ *												trying to marshal the slate request to the gamethread can often introduce thread locks. Setting
+ *												both this and 'LazyInitConnections' to true will allow tools that do not use slate to initialize
+ *												the VA connections on use, but force tools that can display the slate dialog to initialized 
+ *												during preinit on the game thread so that the dialog can be shown. Note that this only overrides
+												the setting of 'LazyInitConnections' via the config file, not cvar, cmdline or code [Default=false]
  * UseLegacyErrorHandling [bool]:				Controls how we deal with errors encountered when pulling payloads. When true a failed payload 
  *												pull will return an error and allow the process to carry on (the original error handling logic)
  *												and when false a dialog will be displayed to the user warning them about the failed pull and 
