Demote some VA logging to verbose to clean up editor shutdown logs

#rnx

diff --git a/Engine/Source/Developer/Virtualization/Private/VirtualizationManager.cpp b/Engine/Source/Developer/Virtualization/Private/VirtualizationManager.cpp
--- a/Engine/Source/Developer/Virtualization/Private/VirtualizationManager.cpp
+++ b/Engine/Source/Developer/Virtualization/Private/VirtualizationManager.cpp
@@ -514,7 +514,7 @@ FVirtualizationManager::~FVirtualizationManager()
 		IConsoleManager::Get().UnregisterConsoleObject(ConsoleObject);
 	}
 
-	UE_LOG(LogVirtualization, Log, TEXT("Destroying backends"));
+	UE_LOG(LogVirtualization, Verbose, TEXT("Destroying backends"));
 
 	CacheStorageBackends.Empty();
 	PersistentStorageBackends.Empty();
@@ -522,7 +522,7 @@ FVirtualizationManager::~FVirtualizationManager()
 
 	AllBackends.Empty(); // This will delete all backends and beyond this point all references to them are invalid
 
-	UE_LOG(LogVirtualization, Log, TEXT("Virtualization manager destroyed"));
+	UE_LOG(LogVirtualization, Verbose, TEXT("Virtualization manager destroyed"));
 }
 
 bool FVirtualizationManager::Initialize(const FInitParams& InitParams)
diff --git a/Engine/Source/Runtime/Core/Private/Virtualization/VirtualizationSystem.cpp b/Engine/Source/Runtime/Core/Private/Virtualization/VirtualizationSystem.cpp
--- a/Engine/Source/Runtime/Core/Private/Virtualization/VirtualizationSystem.cpp
+++ b/Engine/Source/Runtime/Core/Private/Virtualization/VirtualizationSystem.cpp
@@ -360,7 +360,7 @@ bool ShouldInitializePreSlate()
 void Shutdown()
 {
 	GVirtualizationSystem.Reset();
-	UE_LOG(LogVirtualization, Log, TEXT("UE::Virtualization was shutdown"));
+	UE_LOG(LogVirtualization, Verbose, TEXT("UE::Virtualization was shutdown"));
 }
 
 
