Tweak the text in the VA revision control connection error dialog based on feedback.

#rnx

- Removed direct references to perforce from the dialog code, those strings are now passed in by the caller. This should make the dialog more usable by other revision control providers at some point.
- Make it more obvious that the dialog is related to virtualized assets.
- Add more references to the revision control provider name so that people realize which settings they are supposed to be entering.
- Removed some debug code for launching the dialog from a console command.
- Fixed a bug where clicking the button to reset the ini file settings and use the default connection would reset the settings but then not try to reconnect.
- Make better use of the dialog space by increasing the size of the input boxes at the expense of white space, so that longer server addresses can fit on screen.

diff --git a/Engine/Source/Developer/Virtualization/Private/SVirtualizationRevisionControlConnectionDialog.cpp b/Engine/Source/Developer/Virtualization/Private/SVirtualizationRevisionControlConnectionDialog.cpp
--- a/Engine/Source/Developer/Virtualization/Private/SVirtualizationRevisionControlConnectionDialog.cpp
+++ b/Engine/Source/Developer/Virtualization/Private/SVirtualizationRevisionControlConnectionDialog.cpp
@@ -24,7 +24,7 @@
 namespace UE::Virtualization
 {
 
-SRevisionControlConnectionDialog::FResult SRevisionControlConnectionDialog::RunDialog(FStringView CurrentPort, FStringView CurrentUsername)
+SRevisionControlConnectionDialog::FResult SRevisionControlConnectionDialog::RunDialog(FStringView RevisionControlName, FStringView ConfigSectionName,  FStringView CurrentPort, FStringView CurrentUsername)
 {
 	if (FApp::IsUnattended())
 	{
@@ -46,8 +46,10 @@ SRevisionControlConnectionDialog::FResult SRevisionControlConnectionDialog::RunD
 
 	UE_LOG(LogVirtualization, Display, TEXT("Creating dialog"));
 
+	FText WindowTitle = FText::Format(LOCTEXT("VASCSettings", "Virtualized Assets - {0} Revision Control Settings"), FText::FromStringView(RevisionControlName));
+
 	TSharedPtr<SWindow> DialogWindow = SNew(SWindow)
-		.Title(LOCTEXT("VASCSettings", "Perforce Source Control Backend Settings"))
+		.Title(MoveTemp(WindowTitle))
 		.FocusWhenFirstShown(true)
 		.SupportsMinimize(false)
 		.SupportsMaximize(false)
@@ -66,14 +68,14 @@ SRevisionControlConnectionDialog::FResult SRevisionControlConnectionDialog::RunD
 			.AutoHeight()
 			.Padding(16.0f, 16.0f, 16.0f, 0.0f)
 			[
-				SAssignNew(DialogWidget, SRevisionControlConnectionDialog, CurrentPort, CurrentUsername)
+				SAssignNew(DialogWidget, SRevisionControlConnectionDialog, RevisionControlName, ConfigSectionName, CurrentPort, CurrentUsername)
 				.Window(DialogWindow)
 			]
 		];
 
 	DialogWindow->SetContent(DialogWrapper.ToSharedRef());
 
-	UE_LOG(LogVirtualization, Display, TEXT("Connection to source control for virtualized assets failed. Offering the user the choice to retry or continue anyway"));
+	UE_LOG(LogVirtualization, Display, TEXT("Connection to revision control for virtualized assets failed. Offering the user the choice to retry or continue anyway"));
 
 	TSharedPtr<SWindow> ParentWindow = FSlateApplication::Get().GetActiveTopLevelWindow();
 	FSlateApplication::Get().AddModalWindow(DialogWindow.ToSharedRef(), ParentWindow);
@@ -86,14 +88,20 @@ SRevisionControlConnectionDialog::FResult SRevisionControlConnectionDialog::RunD
 	return FResult();
 }
 
-void SRevisionControlConnectionDialog::Construct(const FArguments& InArgs, FStringView CurrentPort, FStringView CurrentUsername)
+void SRevisionControlConnectionDialog::Construct(const FArguments& InArgs, FStringView RevisionControlName, FStringView InConfigSectionName, FStringView CurrentPort, FStringView CurrentUsername)
 {
 	WindowWidget = InArgs._Window;
 
-	const FString CurPort = TEXT("<P4PORT Here>");
-	const FString CurUser = TEXT("<P4USER Here>");
-
 	const FString ConnectionHelpUrl = FVirtualizationManager::GetConnectionHelpUrl();
+	ConfigSectionName = InConfigSectionName;
+
+	FText Message = FText::Format(LOCTEXT("VASCMsg", "Failed to connect to the {0} revision control server.\nThis may prevent you from loading virtualized assets in the future!\n\nPlease enter the correct {1} revision control settings below:"), 
+										FText::FromStringView(RevisionControlName),
+										FText::FromStringView(RevisionControlName));
+
+
+	const FText PortToolTip = FText::Format(LOCTEXT("PortLabel_Tooltip", "The server and port for your {0} server. Usage ServerName:1234."), FText::FromStringView(RevisionControlName));
+	const FText UserToolTip = FText::Format(LOCTEXT("UserNameLabel_Tooltip", "{0} username."), FText::FromStringView(RevisionControlName));
 
 	ChildSlot
 	[
@@ -120,7 +128,7 @@ void SRevisionControlConnectionDialog::Construct(const FArguments& InArgs, FStri
 			+ SHorizontalBox::Slot()
 			[
 				SNew(STextBlock)
-				.Text(LOCTEXT("VASCMsg", "Failed to connect to the source control backend.\nThis may prevent you from accessing virtualized data in the future.\n\nPlease enter the correct source control settings below:"))
+				.Text(MoveTemp(Message))
 				.AutoWrapText(true)
 			]
 		]
@@ -130,7 +138,7 @@ void SRevisionControlConnectionDialog::Construct(const FArguments& InArgs, FStri
 		[
 			SNew(SHorizontalBox)
 			+ SHorizontalBox::Slot()
-			.FillWidth(1.0f)
+			.FillWidth(0.2f)
 			.Padding(FMargin(0.0f, 0.0f, 16.0f, 0.0f))
 			[
 				SNew(SVerticalBox)
@@ -141,7 +149,7 @@ void SRevisionControlConnectionDialog::Construct(const FArguments& InArgs, FStri
 				[
 					SNew(STextBlock)
 					.Text(LOCTEXT("PortLabel", "Server"))
-					.ToolTipText(LOCTEXT("PortLabel_Tooltip", "The server and port for your Perforce server. Usage ServerName:1234."))
+					.ToolTipText(PortToolTip)
 				]
 				+ SVerticalBox::Slot()
 				.Padding(FMargin(0.0f, 0.0f, 0.0f, 10.0f))
@@ -150,11 +158,11 @@ void SRevisionControlConnectionDialog::Construct(const FArguments& InArgs, FStri
 				[
 					SNew(STextBlock)
 					.Text(LOCTEXT("UserNameLabel", "User Name"))
-					.ToolTipText(LOCTEXT("UserNameLabel_Tooltip", "Perforce username."))
+					.ToolTipText(UserToolTip)
 				]
 			]
 			+ SHorizontalBox::Slot()
-			.FillWidth(2.0f)
+			.FillWidth(0.8f)
 			[
 				SNew(SVerticalBox)
 				+ SVerticalBox::Slot()
@@ -163,7 +171,7 @@ void SRevisionControlConnectionDialog::Construct(const FArguments& InArgs, FStri
 				[
 					SAssignNew(PortTextWidget, SEditableTextBox)
 					.Text(FText::FromString(FString(CurrentPort)))
-					.ToolTipText(LOCTEXT("VASC_PortTip", "The server and port for your Perforce server. Usage ServerName:1234."))
+					.ToolTipText(PortToolTip)
 				]
 				+ SVerticalBox::Slot()
 				.Padding(FMargin(0.0f, 0.0f, 0.0f, 10.0f))
@@ -171,7 +179,7 @@ void SRevisionControlConnectionDialog::Construct(const FArguments& InArgs, FStri
 				[
 					SAssignNew(UsernameTextWidget, SEditableTextBox)
 					.Text(FText::FromString(FString(CurrentUsername)))
-					.ToolTipText(LOCTEXT("VASC_UserTip", "Perforce username."))
+					.ToolTipText(UserToolTip)
 				]
 			]
 		]
@@ -268,7 +276,10 @@ void SRevisionControlConnectionDialog::CloseModalDialog()
 
 FReply SRevisionControlConnectionDialog::OnResetToDefaults()
 {
-	GConfig->EmptySection(TEXT("PerforceSourceControl.VirtualizationSettings"), SourceControlHelpers::GetSettingsIni());
+	UE_LOG(LogVirtualization, Display, TEXT("User opted to clear the ini file settings and connect using the revision control defaults"));
+	GConfig->EmptySection(*ConfigSectionName, SourceControlHelpers::GetSettingsIni());
+
+	Result = EResult::Retry;
 
 	Port.Empty();
 	UserName.Empty();
@@ -280,7 +291,7 @@ FReply SRevisionControlConnectionDialog::OnResetToDefaults()
 
 FReply SRevisionControlConnectionDialog::OnRetryConnection()
 {
-	UE_LOG(LogVirtualization, Display, TEXT("User opted to retry connecting to source control"));
+	UE_LOG(LogVirtualization, Display, TEXT("User opted to retry connecting to revision control"));
 
 	Result = EResult::Retry;
 
@@ -294,7 +305,7 @@ FReply SRevisionControlConnectionDialog::OnRetryConnection()
 
 FReply SRevisionControlConnectionDialog::OnSkip()
 {
-	UE_LOG(LogVirtualization, Warning, TEXT("User opted not to connect to source control. Virtualized data may not be accessible!"));
+	UE_LOG(LogVirtualization, Warning, TEXT("User opted not to connect to revision control. Virtualized data may not be accessible!"));
 
 	Result = EResult::Skip;
 
@@ -310,11 +321,6 @@ void SRevisionControlConnectionDialog::OnUrlClicked() const
 	FPlatformProcess::LaunchURL(*ConnectionHelpUrl, nullptr, nullptr);
 }
 
-static FAutoConsoleCommand CCmdTestDialog = FAutoConsoleCommand(
-	TEXT("TestVADialog"),
-	TEXT(""),
-	FConsoleCommandDelegate::CreateStatic(SRevisionControlConnectionDialog::RunDialogCvar));
-
 } // namespace UE::Virtualization
 
 #undef LOCTEXT_NAMESPACE
diff --git a/Engine/Source/Developer/Virtualization/Private/SVirtualizationRevisionControlConnectionDialog.h b/Engine/Source/Developer/Virtualization/Private/SVirtualizationRevisionControlConnectionDialog.h
--- a/Engine/Source/Developer/Virtualization/Private/SVirtualizationRevisionControlConnectionDialog.h
+++ b/Engine/Source/Developer/Virtualization/Private/SVirtualizationRevisionControlConnectionDialog.h
@@ -37,17 +37,13 @@ public:
 		FString UserName;
 	};
 
-	static FResult RunDialog(FStringView CurrentPort, FStringView CurrentUsername);
-	static void RunDialogCvar()
-	{
-		RunDialog(TEXT("<P4PORT Here>"), TEXT("<P4USER Here>"));
-	}
+	static FResult RunDialog(FStringView RevisionControlName, FStringView ConfigSectionName, FStringView CurrentPort, FStringView CurrentUsername);
 
 	SLATE_BEGIN_ARGS(SRevisionControlConnectionDialog) {}
 		SLATE_ARGUMENT(TSharedPtr<SWindow>, Window)
 	SLATE_END_ARGS()
 
-	void Construct(const FArguments& InArgs, FStringView CurrentPort, FStringView CurrentUsername);
+	void Construct(const FArguments& InArgs, FStringView RevisionControlName, FStringView ConfigSectionName, FStringView CurrentPort, FStringView CurrentUsername);
 
 	enum class EResult : uint8
 	{
@@ -87,6 +83,8 @@ private:
 
 	EResult Result = EResult::Skip;
 
+	FString ConfigSectionName;
+
 	FString Port;
 	FString UserName;
 };
diff --git a/Engine/Source/Developer/Virtualization/Private/VirtualizationSourceControlBackend.cpp b/Engine/Source/Developer/Virtualization/Private/VirtualizationSourceControlBackend.cpp
--- a/Engine/Source/Developer/Virtualization/Private/VirtualizationSourceControlBackend.cpp
+++ b/Engine/Source/Developer/Virtualization/Private/VirtualizationSourceControlBackend.cpp
@@ -368,7 +368,11 @@ IVirtualizationBackend::EConnectionStatus FSourceControlBackend::OnConnect()
 		// TODO: Maybe give a bespoke error at this point?
 		if (bUseRetryConnectionDialog && bUseLocalIniFileSettings)
 		{
-			SRevisionControlConnectionDialog::FResult DialogResult = SRevisionControlConnectionDialog::RunDialog(Port, UserName);
+			SRevisionControlConnectionDialog::FResult DialogResult = SRevisionControlConnectionDialog::RunDialog(	TEXT("Perforce"),
+																													TEXT("PerforceSourceControl.VirtualizationSettings"),
+																													Port,
+																													UserName);
+
 			if (!DialogResult.bShouldRetry)
 			{
 				OnConnectionError(ErrorMessage);
