Update the source control backend reconnection dialog to show the connection errors.

#rnx

- When a connection attempt fails we now return the credentials used and the errors given by the system.
-- The credentials will be used to auto fill out the dialog, letting the user see what credentials were attempted in the last attempt.
-- The errors are shown in a box with red error text, letting them know what went wrong.
- We no longer let the revision control provider log errors when establishing a connection, so if the users correct info and successfully connects on a subsequent attempt they will not see an error in the logs.
-- If the connection completely fails we still log as before.

diff --git a/Engine/Source/Developer/Virtualization/Private/SVirtualizationRevisionControlConnectionDialog.cpp b/Engine/Source/Developer/Virtualization/Private/SVirtualizationRevisionControlConnectionDialog.cpp
--- a/Engine/Source/Developer/Virtualization/Private/SVirtualizationRevisionControlConnectionDialog.cpp
+++ b/Engine/Source/Developer/Virtualization/Private/SVirtualizationRevisionControlConnectionDialog.cpp
@@ -11,6 +11,7 @@
 #include "Misc/App.h"
 #include "Misc/ConfigCacheIni.h"
 #include "SourceControlHelpers.h"
+#include "Styling/StyleColors.h"
 #include "Widgets/Images/SImage.h"
 #include "Widgets/Input/SButton.h"
 #include "Widgets/Input/SEditableTextBox.h"
@@ -24,7 +25,7 @@
 namespace UE::Virtualization
 {
 
-SRevisionControlConnectionDialog::FResult SRevisionControlConnectionDialog::RunDialog(FStringView RevisionControlName, FStringView ConfigSectionName,  FStringView CurrentPort, FStringView CurrentUsername)
+SRevisionControlConnectionDialog::FResult SRevisionControlConnectionDialog::RunDialog(FStringView RevisionControlName, FStringView ConfigSectionName,  FStringView CurrentPort, FStringView CurrentUsername, const FText& ErrorMessage)
 {
 	if (FApp::IsUnattended())
 	{
@@ -60,7 +61,7 @@ SRevisionControlConnectionDialog::FResult SRevisionControlConnectionDialog::RunD
 
 	TSharedPtr<SBorder> DialogWrapper =
 		SNew(SBorder)
-		.BorderImage(FAppStyle::GetBrush("ToolPanel.GroupBorder"))
+	//	.BorderImage(FAppStyle::GetBrush("ToolPanel.GroupBorder"))
 		.Padding(4.0f)
 		[
 			SNew(SVerticalBox)
@@ -68,7 +69,7 @@ SRevisionControlConnectionDialog::FResult SRevisionControlConnectionDialog::RunD
 			.AutoHeight()
 			.Padding(16.0f, 16.0f, 16.0f, 0.0f)
 			[
-				SAssignNew(DialogWidget, SRevisionControlConnectionDialog, RevisionControlName, ConfigSectionName, CurrentPort, CurrentUsername)
+				SAssignNew(DialogWidget, SRevisionControlConnectionDialog, RevisionControlName, ConfigSectionName, CurrentPort, CurrentUsername, ErrorMessage)
 				.Window(DialogWindow)
 			]
 		];
@@ -88,20 +89,20 @@ SRevisionControlConnectionDialog::FResult SRevisionControlConnectionDialog::RunD
 	return FResult();
 }
 
-void SRevisionControlConnectionDialog::Construct(const FArguments& InArgs, FStringView RevisionControlName, FStringView InConfigSectionName, FStringView CurrentPort, FStringView CurrentUsername)
+void SRevisionControlConnectionDialog::Construct(const FArguments& InArgs, FStringView RevisionControlName, FStringView InConfigSectionName, FStringView CurrentPort, FStringView CurrentUsername, const FText& ErrorMessage)
 {
 	WindowWidget = InArgs._Window;
 
 	const FString ConnectionHelpUrl = FVirtualizationManager::GetConnectionHelpUrl();
 	ConfigSectionName = InConfigSectionName;
 
-	FText Message = FText::Format(LOCTEXT("VASCMsg", "Failed to connect to the {0} revision control server.\nThis may prevent you from loading virtualized assets in the future!\n\nPlease enter the correct {1} revision control settings below:"), 
-										FText::FromStringView(RevisionControlName),
-										FText::FromStringView(RevisionControlName));
-
-
+	FText MessagePt1 = FText::Format(LOCTEXT("VASCMsg", "Failed to connect to the {0} revision control server with the following errors:"), FText::FromStringView(RevisionControlName));
+	FText MessagePt2 = FText::Format(LOCTEXT("VASCMsg", "This may prevent you from loading virtualized assets in the future!\nPlease enter the correct {0} revision control settings below:"), FText::FromStringView(RevisionControlName));
+	
 	const FText PortToolTip = FText::Format(LOCTEXT("PortLabel_Tooltip", "The server and port for your {0} server. Usage ServerName:1234."), FText::FromStringView(RevisionControlName));
 	const FText UserToolTip = FText::Format(LOCTEXT("UserNameLabel_Tooltip", "{0} username."), FText::FromStringView(RevisionControlName));
+	
+	const char* BorderImage = "BlackBrush";
 
 	ChildSlot
 	[
@@ -127,9 +128,37 @@ void SRevisionControlConnectionDialog::Construct(const FArguments& InArgs, FStri
 			]
 			+ SHorizontalBox::Slot()
 			[
-				SNew(STextBlock)
-				.Text(MoveTemp(Message))
-				.AutoWrapText(true)
+				SNew(SVerticalBox)
+				+ SVerticalBox::Slot()
+				.Padding(FMargin(0.0f, 0.0f, 0.0f, 0.0f))
+				.AutoHeight()
+				[
+					SNew(STextBlock)
+					.Text(MoveTemp(MessagePt1))
+					.AutoWrapText(true)
+				]
+				+ SVerticalBox::Slot()
+				.Padding(FMargin(0.0f, 16.0f, 0.0f, 16.0f))
+				.AutoHeight()
+				[
+					SNew(SBorder)
+					.BorderImage(FAppStyle::Get().GetBrush(BorderImage))
+					[
+						SNew(STextBlock)
+						.ColorAndOpacity(FSlateColor(EStyleColor::Error))
+						.Text(ErrorMessage)
+						.AutoWrapText(true)
+					]
+				]
+				+ SVerticalBox::Slot()
+				.Padding(FMargin(0.0f, 0.0f, 0.0f, 0.0f))
+				.AutoHeight()
+				[
+					SNew(STextBlock)
+					.Text(MoveTemp(MessagePt2))
+					.AutoWrapText(true)
+				]
+				
 			]
 		]
 		+ SVerticalBox::Slot()
@@ -240,7 +269,6 @@ void SRevisionControlConnectionDialog::Construct(const FArguments& InArgs, FStri
 				SNew(SHorizontalBox)
 				+ SHorizontalBox::Slot()
 				.AutoWidth()
-				.Padding(FMargin(5.0f, 0.0f))
 				[
 					SNew(SButton)
 					.VAlign(VAlign_Center)
@@ -252,7 +280,7 @@ void SRevisionControlConnectionDialog::Construct(const FArguments& InArgs, FStri
 				]
 				+ SHorizontalBox::Slot()
 				.AutoWidth()
-				.Padding(FMargin(5.0f, 0.0f))
+				.Padding(FMargin(10.0f, 0.0f, 0.0f, 0.0f))
 				[
 					SNew(SButton)
 					.VAlign(VAlign_Center)
diff --git a/Engine/Source/Developer/Virtualization/Private/SVirtualizationRevisionControlConnectionDialog.h b/Engine/Source/Developer/Virtualization/Private/SVirtualizationRevisionControlConnectionDialog.h
--- a/Engine/Source/Developer/Virtualization/Private/SVirtualizationRevisionControlConnectionDialog.h
+++ b/Engine/Source/Developer/Virtualization/Private/SVirtualizationRevisionControlConnectionDialog.h
@@ -37,13 +37,13 @@ public:
 		FString UserName;
 	};
 
-	static FResult RunDialog(FStringView RevisionControlName, FStringView ConfigSectionName, FStringView CurrentPort, FStringView CurrentUsername);
+	static FResult RunDialog(FStringView RevisionControlName, FStringView ConfigSectionName, FStringView CurrentPort, FStringView CurrentUsername, const FText& ErrorMessage);
 
 	SLATE_BEGIN_ARGS(SRevisionControlConnectionDialog) {}
 		SLATE_ARGUMENT(TSharedPtr<SWindow>, Window)
 	SLATE_END_ARGS()
 
-	void Construct(const FArguments& InArgs, FStringView RevisionControlName, FStringView ConfigSectionName, FStringView CurrentPort, FStringView CurrentUsername);
+	void Construct(const FArguments& InArgs, FStringView RevisionControlName, FStringView ConfigSectionName, FStringView CurrentPort, FStringView CurrentUsername, const FText& ErrorMessage);
 
 	enum class EResult : uint8
 	{
diff --git a/Engine/Source/Developer/Virtualization/Private/VirtualizationSourceControlBackend.cpp b/Engine/Source/Developer/Virtualization/Private/VirtualizationSourceControlBackend.cpp
--- a/Engine/Source/Developer/Virtualization/Private/VirtualizationSourceControlBackend.cpp
+++ b/Engine/Source/Developer/Virtualization/Private/VirtualizationSourceControlBackend.cpp
@@ -340,7 +340,8 @@ IVirtualizationBackend::EConnectionStatus FSourceControlBackend::OnConnect()
 	}
 
 	FString Port = ServerAddress;
-	FString UserName = TEXT("");
+	FString UserName;
+
 	bool bSaveSettings = false; // Initially we will try to reason the perforce settings from the ini file and global environment
 								// so we don't want to save these values to the users ini files.
 
@@ -371,7 +372,8 @@ IVirtualizationBackend::EConnectionStatus FSourceControlBackend::OnConnect()
 			SRevisionControlConnectionDialog::FResult DialogResult = SRevisionControlConnectionDialog::RunDialog(	TEXT("Perforce"),
 																													TEXT("PerforceSourceControl.VirtualizationSettings"),
 																													Port,
-																													UserName);
+																													UserName,
+																													ErrorMessage);
 
 			if (!DialogResult.bShouldRetry)
 			{
@@ -393,7 +395,7 @@ IVirtualizationBackend::EConnectionStatus FSourceControlBackend::OnConnect()
 }
 
 
-IVirtualizationBackend::EConnectionStatus FSourceControlBackend::OnConnectInternal(FStringView Port, FStringView Username, bool bSaveConnectionSettings, FText& OutErrorMessage)
+IVirtualizationBackend::EConnectionStatus FSourceControlBackend::OnConnectInternal(FString& InOutPort, FString& InOutUsername, bool bSaveConnectionSettings, FText& OutErrorMessage)
 {
 	// We do not want the connection to have a client workspace so explicitly set it to empty
 	FSourceControlInitSettings SCCSettings(FSourceControlInitSettings::EBehavior::OverrideExisting);
@@ -408,16 +410,17 @@ IVirtualizationBackend::EConnectionStatus FSourceControlBackend::OnConnectIntern
 
 	SCCSettings.SetConfigBehavior(IniBehavior);
 
-	if (!Port.IsEmpty())
+	if (!InOutPort.IsEmpty())
 	{
-		SCCSettings.AddSetting(TEXT("P4Port"), Port);
+		SCCSettings.AddSetting(TEXT("P4Port"), InOutPort);
 	}
 
-	if (!Username.IsEmpty())
+	if (!InOutUsername.IsEmpty())
 	{
-		SCCSettings.AddSetting(TEXT("P4User"), Username);
+		SCCSettings.AddSetting(TEXT("P4User"), InOutUsername);
 	}
 
+	// By setting an empty "P4Client" we ensure that the provider will not attempt to find a default one.
 	SCCSettings.AddSetting(TEXT("P4Client"), TEXT(""));
 
 	SCCProvider = ISourceControlModule::Get().CreateProvider(FName("Perforce"), TEXT("Virtualization"), SCCSettings);
@@ -427,14 +430,36 @@ IVirtualizationBackend::EConnectionStatus FSourceControlBackend::OnConnectIntern
 		return IVirtualizationBackend::EConnectionStatus::Error;
 	}
 
-	// Will attempt to make initial connection to the server, it doesn't return any error values we can check against but
-	// it will output problems to the log file.
-	const bool bForceConnection = true;
-	SCCProvider->Init(bForceConnection);
+	const ISourceControlProvider::EInitFlags InitOptions =	ISourceControlProvider::EInitFlags::AttemptConnection |
+															ISourceControlProvider::EInitFlags::SupressErrorLogging;
+
+	ISourceControlProvider::FInitResult Result = SCCProvider->Init(InitOptions);
+
+	if (FString* Value = Result.ConnectionSettings.Find(ISourceControlProvider::EStatus::Port))
+	{
+		InOutPort = *Value;
+	}
+
+	if (FString* Value = Result.ConnectionSettings.Find(ISourceControlProvider::EStatus::User))
+	{
+		InOutUsername = *Value;
+	}
 
-	if (!SCCProvider->IsAvailable())
+	if (!Result.bIsAvailable)
 	{
-		OutErrorMessage = FText(LOCTEXT("FailedSourceControlConnection", "Failed to connect to revision control backend, see the Message Log 'Revision Control' errors for details.\nThe revision control backend will be unable to pull payloads, is your revision control config set up correctly?"));
+		FTextBuilder MsgBuilder;
+		if (!Result.Errors.ErrorMessage.IsEmpty())
+		{
+			MsgBuilder.AppendLine(Result.Errors.ErrorMessage);
+		}
+
+		for (const FText& Error : Result.Errors.AdditionalErrors)
+		{
+			MsgBuilder.AppendLine(Error);
+		}
+
+		OutErrorMessage = MsgBuilder.ToText();
+		
 		return IVirtualizationBackend::EConnectionStatus::Error;
 	}
 
@@ -1183,7 +1208,7 @@ bool FSourceControlBackend::FindSubmissionWorkingDir(const FString& ConfigEntry)
 	}
 }
 
-void FSourceControlBackend::OnConnectionError(FText Message)
+void FSourceControlBackend::OnConnectionError(FText ErrorMessage)
 {
 	// We don't know when or where this error might occur. We can currently only create 
 	// FMessageLog on the game thread or risk slate asserts, and if we raise a 
@@ -1195,13 +1220,15 @@ void FSourceControlBackend::OnConnectionError(FText Message)
 	// we will write it to the log at the point it is raised rather than the point it is
 	// deferred to in an attempt to make the log more readable.
 
-	if (Message.IsEmpty())
+	if (ErrorMessage.IsEmpty())
 	{
 		return;
 	}
 
+	const FText UserMessage(LOCTEXT("FailedSourceControlConnection", "Failed to connect to revision control backend, see the Message Log 'Revision Control' errors for details.\nThe revision control backend will be unable to pull payloads, is your revision control config set up correctly?"));
+
 	TSharedRef<FTokenizedMessage> TokenizedMsg = FTokenizedMessage::Create(EMessageSeverity::Error);
-	TokenizedMsg->AddToken(FTextToken::Create(Message));
+	TokenizedMsg->AddToken(FTextToken::Create(UserMessage));
 
 	FString ConnectionHelpUrl = FVirtualizationManager::GetConnectionHelpUrl();
 	if (!ConnectionHelpUrl.IsEmpty())
@@ -1212,20 +1239,37 @@ void FSourceControlBackend::OnConnectionError(FText Message)
 	if (::IsInGameThread())
 	{
 		// If we are on the game thread we can post the error message immediately
-		FMessageLog Log("LogVirtualization");
-		Log.AddMessage(TokenizedMsg);
+
+		{
+			FMessageLog Log("LogVirtualization");
+			Log.Error(ErrorMessage);
+		}
+
+		{
+			FMessageLog Log("LogVirtualization");
+			Log.AddMessage(TokenizedMsg);
+		}
 	}
 	else
 	{
 		// We can only send a FMessageLog on the GameThread so for now just log the error
 		// and we can send it to the FMessageLog system next tick
-		UE_LOG(LogVirtualization, Error, TEXT("%s"), *Message.ToString());
+		UE_LOG(LogVirtualization, Error, TEXT("%s"), *ErrorMessage.ToString());
+		UE_LOG(LogVirtualization, Error, TEXT("%s"), *UserMessage.ToString());
 
-		auto Callback = [TokenizedMsg, bShouldNotify = !bSuppressNotifications](float Delta)->bool
+		auto Callback = [ErrorMessage, TokenizedMsg, bShouldNotify = !bSuppressNotifications](float Delta)->bool
 		{
-			FMessageLog Log("LogVirtualization");
-			Log.SuppressLoggingToOutputLog();
-			Log.AddMessage(TokenizedMsg);
+			{
+				FMessageLog Log("LogVirtualization");
+				Log.SuppressLoggingToOutputLog();
+				Log.Error(ErrorMessage);
+			}
+
+			{
+				FMessageLog Log("LogVirtualization");
+				Log.SuppressLoggingToOutputLog();
+				Log.AddMessage(TokenizedMsg);
+			}
 
 			return false;
 		};
@@ -1239,7 +1283,7 @@ void FSourceControlBackend::OnConnectionError(FText Message)
 		// be issued on the first frame, increasing the chance that the user will see it.
 		auto NotificationCallback = [](float Delta)->bool
 			{
-				FMessageLog Notification("LogVirtualization");;
+				FMessageLog Notification("LogVirtualization");
 				Notification.SuppressLoggingToOutputLog();
 
 				Notification.Notify(LOCTEXT("ConnectionError", "Asset virtualization connection errors were encountered, see the message log for more info"));
diff --git a/Engine/Source/Developer/Virtualization/Private/VirtualizationSourceControlBackend.h b/Engine/Source/Developer/Virtualization/Private/VirtualizationSourceControlBackend.h
--- a/Engine/Source/Developer/Virtualization/Private/VirtualizationSourceControlBackend.h
+++ b/Engine/Source/Developer/Virtualization/Private/VirtualizationSourceControlBackend.h
@@ -129,7 +129,7 @@ private:
 	virtual bool Initialize(const FString& ConfigEntry) override;
 
 	virtual EConnectionStatus OnConnect() override;
-	IVirtualizationBackend::EConnectionStatus OnConnectInternal(FStringView Port, FStringView Username, bool bSaveConnectionSettings, FText& OutErrorMessage);
+	IVirtualizationBackend::EConnectionStatus OnConnectInternal(FString& InOutPort, FString& InOutUsername, bool bSaveConnectionSettings, FText& OutErrorMessage);
 	
 	virtual bool PushData(TArrayView<FPushRequest> Requests, EPushFlags Flags) override;
 	virtual bool PullData(TArrayView<FPullRequest> Requests, EPullFlags Flags, FText& OutErrors) override;
@@ -147,7 +147,7 @@ private:
 	bool FindSubmissionWorkingDir(const FString& ConfigEntry);
 
 	/** Will display a FMessage notification to the user on the next valid engine tick to try and keep them aware of connection failures */
-	void OnConnectionError(FText Message);
+	void OnConnectionError(FText ErrorMessage);
 
 	/** A source control connection owned by the backend*/
 	TUniquePtr<ISourceControlProvider> SCCProvider;
