Fix duplicate localization key

#rnx

- Also removed some commented out code and a pointless local variable as code cleanup while I was on.

diff --git a/Engine/Source/Developer/Virtualization/Private/SVirtualizationRevisionControlConnectionDialog.cpp b/Engine/Source/Developer/Virtualization/Private/SVirtualizationRevisionControlConnectionDialog.cpp
--- a/Engine/Source/Developer/Virtualization/Private/SVirtualizationRevisionControlConnectionDialog.cpp
+++ b/Engine/Source/Developer/Virtualization/Private/SVirtualizationRevisionControlConnectionDialog.cpp
@@ -61,7 +61,6 @@ SRevisionControlConnectionDialog::FResult SRevisionControlConnectionDialog::RunD
 
 	TSharedPtr<SBorder> DialogWrapper =
 		SNew(SBorder)
-	//	.BorderImage(FAppStyle::GetBrush("ToolPanel.GroupBorder"))
 		.Padding(4.0f)
 		[
 			SNew(SVerticalBox)
@@ -96,13 +95,11 @@ void SRevisionControlConnectionDialog::Construct(const FArguments& InArgs, FStri
 	const FString ConnectionHelpUrl = FVirtualizationManager::GetConnectionHelpUrl();
 	ConfigSectionName = InConfigSectionName;
 
-	FText MessagePt1 = FText::Format(LOCTEXT("VASCMsg", "Failed to connect to the {0} revision control server with the following errors:"), FText::FromStringView(RevisionControlName));
-	FText MessagePt2 = FText::Format(LOCTEXT("VASCMsg", "This may prevent you from loading virtualized assets in the future!\nPlease enter the correct {0} revision control settings below:"), FText::FromStringView(RevisionControlName));
+	FText MessagePt1 = FText::Format(LOCTEXT("VASCMsgPt1", "Failed to connect to the {0} revision control server with the following errors:"), FText::FromStringView(RevisionControlName));
+	FText MessagePt2 = FText::Format(LOCTEXT("VASCMsgPt2", "This may prevent you from loading virtualized assets in the future!\nPlease enter the correct {0} revision control settings below:"), FText::FromStringView(RevisionControlName));
 	
 	const FText PortToolTip = FText::Format(LOCTEXT("PortLabel_Tooltip", "The server and port for your {0} server. Usage ServerName:1234."), FText::FromStringView(RevisionControlName));
 	const FText UserToolTip = FText::Format(LOCTEXT("UserNameLabel_Tooltip", "{0} username."), FText::FromStringView(RevisionControlName));
-	
-	const char* BorderImage = "BlackBrush";
 
 	ChildSlot
 	[
@@ -142,7 +139,7 @@ void SRevisionControlConnectionDialog::Construct(const FArguments& InArgs, FStri
 				.AutoHeight()
 				[
 					SNew(SBorder)
-					.BorderImage(FAppStyle::Get().GetBrush(BorderImage))
+					.BorderImage(FAppStyle::Get().GetBrush("BlackBrush"))
 					[
 						SNew(STextBlock)
 						.ColorAndOpacity(FSlateColor(EStyleColor::Error))
