Prevent the VA error dialogs from being shown if the code is being run in a scenario where a user will not be present.

#rnx

- Add a new ::IsProcessInteractive function to UE::Virtualization::Utils namespace that will return true if we think the current process is being run in a scenario where a user will be present and that it is okay to show an interactive error dialog.
- Use this new check to make sure that we do not try to display the "connection error" dialog on start up (even though as a slate dialog this will usually fail to display in non-interactive scenarios anyway)
- Use the new check to avoid displaying the "payload pull failure" dialog and default to the process shutting down instead.
-- Added some code comments to explain why we force the process to exit rather than logging a fatal error.

diff --git a/Engine/Source/Developer/Virtualization/Private/VirtualizationManager.cpp b/Engine/Source/Developer/Virtualization/Private/VirtualizationManager.cpp
--- a/Engine/Source/Developer/Virtualization/Private/VirtualizationManager.cpp
+++ b/Engine/Source/Developer/Virtualization/Private/VirtualizationManager.cpp
@@ -20,6 +20,7 @@
 #include "PackageVirtualizationProcess.h"
 #include "ProfilingDebugging/CookStats.h"
 #include "VirtualizationFilterSettings.h"
+#include "VirtualizationUtilities.h"
 
 #define LOCTEXT_NAMESPACE "Virtualization"
 
@@ -1982,12 +1983,11 @@ FVirtualizationManager::ErrorHandlingResult FVirtualizationManager::OnPayloadPul
 
 	if (CriticalSection.TryLock())
 	{
-		
 		const FText Title(LOCTEXT("VAPullTitle", "Failed to pull virtualized data!"));
 
 		FTextBuilder MsgBuilder;
 		MsgBuilder.AppendLine(LOCTEXT("VAPullMsgHeader", "Failed to pull payload(s) from virtualization storage and allowing the editor to continue could corrupt data!"));
-		
+
 		if (!BackendErrors.IsEmpty())
 		{
 			MsgBuilder.AppendLine(FString(TEXT("")));
@@ -2000,16 +2000,30 @@ FVirtualizationManager::ErrorHandlingResult FVirtualizationManager::OnPayloadPul
 			MsgBuilder.AppendLine(PullErrorAdditionalMsg);
 		}
 
-		MsgBuilder.AppendLine(FString(TEXT("")));
-		MsgBuilder.AppendLine(LOCTEXT("VAPullMsgYes", "[Yes] Retry pulling the data"));
-		MsgBuilder.AppendLine(LOCTEXT("VAPullMsgNo", "[No] Quit the editor"));
+		// By default we should quit the process unless the user can opt to retry
+		EAppReturnType::Type Result = EAppReturnType::No;
+
+		if (Utils::IsProcessInteractive())
+		{
+			MsgBuilder.AppendLine(FString(TEXT("")));
+			MsgBuilder.AppendLine(LOCTEXT("VAPullMsgYes", "[Yes] Retry pulling the data"));
+			MsgBuilder.AppendLine(LOCTEXT("VAPullMsgNo", "[No] Quit the editor"));
 
-		const FText Message = MsgBuilder.ToText();
+			const FText Message = MsgBuilder.ToText();
 
-		EAppReturnType::Type Result = FMessageDialog::Open(EAppMsgType::YesNo, EAppReturnType::No, Message, Title);
+			Result = FMessageDialog::Open(EAppMsgType::YesNo, EAppReturnType::No, Message, Title);
+		}
+		else
+		{
+			const FText Message = MsgBuilder.ToText();
+			UE_LOG(LogVirtualization, Error, TEXT("%s"), *Message.ToString());
+		}
 
 		if (Result == EAppReturnType::No)
 		{
+			// Reporting a fatal error will print the callstack and initiate the crash handling system to be treated like a bug
+			// where as this error indicates an infrastructure failure or other connection issue which needs to be solved.
+			// So we will force the process to close after logging our errors instead.
 			UE_LOG(LogVirtualization, Error, TEXT("Failed to pull payloads from persistent storage and the connection could not be re-established, exiting..."));
 			GIsCriticalError = 1;
 			FPlatformMisc::RequestExit(true);
@@ -2019,6 +2033,8 @@ FVirtualizationManager::ErrorHandlingResult FVirtualizationManager::OnPayloadPul
 	}
 	else
 	{
+		// Since we failed to get the lock the error is being dealt with on another thread so we need to wait here until
+		// that thread has resolved the problem or terminated the process.
 		FScopeLock _(&CriticalSection);
 	}
 
diff --git a/Engine/Source/Developer/Virtualization/Private/VirtualizationSourceControlBackend.cpp b/Engine/Source/Developer/Virtualization/Private/VirtualizationSourceControlBackend.cpp
--- a/Engine/Source/Developer/Virtualization/Private/VirtualizationSourceControlBackend.cpp
+++ b/Engine/Source/Developer/Virtualization/Private/VirtualizationSourceControlBackend.cpp
@@ -367,7 +367,7 @@ IVirtualizationBackend::EConnectionStatus FSourceControlBackend::OnConnect()
 		// If the local ini settings are ignored then there is no point saving correct settings given by the user.
 		// They will need to fix their root problem instead.
 		// TODO: Maybe give a bespoke error at this point?
-		if (bUseRetryConnectionDialog && bUseLocalIniFileSettings)
+		if (bUseRetryConnectionDialog && bUseLocalIniFileSettings && Utils::IsProcessInteractive())
 		{
 			SRevisionControlConnectionDialog::FResult DialogResult = SRevisionControlConnectionDialog::RunDialog(	TEXT("Perforce"),
 																													TEXT("PerforceSourceControl.VirtualizationSettings"),
diff --git a/Engine/Source/Developer/Virtualization/Private/VirtualizationUtilities.cpp b/Engine/Source/Developer/Virtualization/Private/VirtualizationUtilities.cpp
--- a/Engine/Source/Developer/Virtualization/Private/VirtualizationUtilities.cpp
+++ b/Engine/Source/Developer/Virtualization/Private/VirtualizationUtilities.cpp
@@ -4,6 +4,7 @@
 
 #include "HAL/PlatformProcess.h"
 #include "IO/IoHash.h"
+#include "Misc/App.h"
 #include "Misc/Paths.h"
 #include "Misc/StringBuilder.h"
 #include "UObject/PackageFileSummary.h"
@@ -139,4 +140,9 @@ bool ExpandEnvironmentVariables(FStringView InputPath, FStringBuilderBase& OutEx
 	}
 }
 
+bool IsProcessInteractive()
+{
+	return !FApp::IsUnattended() && !IsRunningCommandlet() && !GIsRunningUnattendedScript && !IS_PROGRAM;
+}
+
 } // namespace UE::Virtualization::Utils
diff --git a/Engine/Source/Developer/Virtualization/Public/VirtualizationUtilities.h b/Engine/Source/Developer/Virtualization/Public/VirtualizationUtilities.h
--- a/Engine/Source/Developer/Virtualization/Public/VirtualizationUtilities.h
+++ b/Engine/Source/Developer/Virtualization/Public/VirtualizationUtilities.h
@@ -77,4 +77,10 @@ ETrailerFailedReason FindTrailerFailedReason(const FPackagePath& PackagePath);
  */
 bool ExpandEnvironmentVariables(FStringView InputPath, FStringBuilderBase& OutExpandedPath);
 
+/** 
+ * Returns true if the process is interactive and we should display error dialogs to
+ * the user.
+ * */
+bool IsProcessInteractive();
+
 } // namespace UE::Virtualization::Utils
