Improve error reporting with the new error flow in VA

#rnx

- With the new error flow we block the process on a failed pull and then quit the process if the connection cannot be recovered, this means that we never reach the code that reports which payloads failed.
- Added similar reporting to FVirtualizationManager::OnPayloadPullError so we can get a log of which payloads failed to pull.
-- A pull failure should not really be dependant on the payload itself so this info isn't that important but it can be useful to eliminate potential issues by allowing us to confirm that the payload does in fact exist in persistent storage and the user was just unable to access it.

diff --git a/Engine/Source/Developer/Virtualization/Private/VirtualizationManager.cpp b/Engine/Source/Developer/Virtualization/Private/VirtualizationManager.cpp
--- a/Engine/Source/Developer/Virtualization/Private/VirtualizationManager.cpp
+++ b/Engine/Source/Developer/Virtualization/Private/VirtualizationManager.cpp
@@ -197,6 +197,21 @@ public:
 		return CurrentRequests;
 	}
 
+	TArray<FIoHash> GetFailedPayloads() const
+	{
+		TArray<FIoHash> FailedPayloads;
+
+		for (const FPullRequest& Request : CurrentRequests)
+		{
+			if (!Request.IsSuccess())
+			{
+				FailedPayloads.Add(Request.GetIdentifier());
+			}
+		}
+
+		return FailedPayloads;
+	}
+
 	/** Returns if there are still requests that need servicing or not */
 	bool IsWorkComplete() const
 	{
@@ -1941,7 +1956,7 @@ void FVirtualizationManager::PullDataFromAllBackends(TArrayView<FPullRequest> Re
 		{
 			return; // All payloads pulled
 		}
-		else if (OnPayloadPullError(BackendErrors) != ErrorHandlingResult::Retry)
+		else if (OnPayloadPullError(RequestsCollection, BackendErrors) != ErrorHandlingResult::Retry)
 		{
 			return; // Some payloads failed to pull
 		}
@@ -1972,13 +1987,18 @@ void FVirtualizationManager::PullDataFromBackend(IVirtualizationBackend& Backend
 #endif //ENABLE_COOK_STATS
 }
 
-FVirtualizationManager::ErrorHandlingResult FVirtualizationManager::OnPayloadPullError(FStringView BackendErrors)
+FVirtualizationManager::ErrorHandlingResult FVirtualizationManager::OnPayloadPullError(const FPullRequestCollection& Requests, FStringView BackendErrors) const
 {
 	if (bUseLegacyErrorHandling)
 	{
 		return ErrorHandlingResult::AcceptFailedPayloads;
 	}
 
+	for (const FIoHash& FailedPayload : Requests.GetFailedPayloads())
+	{
+		UE_LOG(LogVirtualization, Error, TEXT("Payload '%s' failed to be pulled from any backend'"), *LexToString(FailedPayload));
+	}
+
 	static FCriticalSection CriticalSection;
 
 	if (CriticalSection.TryLock())
diff --git a/Engine/Source/Developer/Virtualization/Private/VirtualizationManager.h b/Engine/Source/Developer/Virtualization/Private/VirtualizationManager.h
--- a/Engine/Source/Developer/Virtualization/Private/VirtualizationManager.h
+++ b/Engine/Source/Developer/Virtualization/Private/VirtualizationManager.h
@@ -243,7 +243,7 @@ private:
 		AcceptFailedPayloads
 	};
 
-	ErrorHandlingResult OnPayloadPullError(FStringView BackendErrors);
+	ErrorHandlingResult OnPayloadPullError(const FPullRequestCollection& Requests, FStringView BackendErrors) const;
 	
 	bool ShouldVirtualizeAsset(const UObject* Owner) const;
 
