Do not use/include the MessageLog module in the Virtualization module if the target is not building for slate.

#rnx

- We only needed the MessageLog module to setup a VA tab in the message log slate window which only matters if we are actually building for a target that uses slate.

diff --git a/Engine/Source/Developer/Virtualization/Private/Virtualization.cpp b/Engine/Source/Developer/Virtualization/Private/Virtualization.cpp
--- a/Engine/Source/Developer/Virtualization/Private/Virtualization.cpp
+++ b/Engine/Source/Developer/Virtualization/Private/Virtualization.cpp
@@ -3,9 +3,12 @@
 #include "Modules/ModuleManager.h"
 
 #include "Features/IModularFeatures.h"
-#include "MessageLogModule.h"
 #include "VirtualizationSourceControlUtilities.h"
 
+#if UE_VA_WITH_SLATE
+#include "MessageLogModule.h"
+#endif //UE_VA_WITH_SLATE
+
 #define LOCTEXT_NAMESPACE "Virtualization"
 
 namespace UE::Virtualization
@@ -18,18 +21,22 @@ public:
 	{
 		IModularFeatures::Get().RegisterModularFeature(FName("VirtualizationSourceControlUtilities"), &SourceControlutility);
 
+#if UE_VA_WITH_SLATE
 		FMessageLogModule& MessageLogModule = FModuleManager::LoadModuleChecked<FMessageLogModule>("MessageLog");
 		MessageLogModule.RegisterLogListing("LogVirtualization", LOCTEXT("AssetVirtualizationLogLabel", "Asset Virtualization"));
+#endif //UE_VA_WITH_SLATE
 	}
 
 	virtual void ShutdownModule() override
 	{
+#if UE_VA_WITH_SLATE
 		if (FModuleManager::Get().IsModuleLoaded("MessageLog"))
 		{
 			FMessageLogModule& MessageLogModule = FModuleManager::GetModuleChecked<FMessageLogModule>("MessageLog");
 			MessageLogModule.UnregisterLogListing("LogVirtualization&");
 		}
-		
+#endif //UE_VA_WITH_SLATE	
+
 		IModularFeatures::Get().UnregisterModularFeature(FName("VirtualizationSourceControlUtilities"), &SourceControlutility);
 	}
 
diff --git a/Engine/Source/Developer/Virtualization/Virtualization.Build.cs b/Engine/Source/Developer/Virtualization/Virtualization.Build.cs
--- a/Engine/Source/Developer/Virtualization/Virtualization.Build.cs
+++ b/Engine/Source/Developer/Virtualization/Virtualization.Build.cs
@@ -12,21 +12,24 @@ public class Virtualization : ModuleRules
 		PrivateIncludePathModuleNames.Add("Analytics");
 
 		PrivateDependencyModuleNames.AddRange(
-			new string[]
-			{
+			new string[] {
 				"Core",
 				"CoreUObject",
 				"DerivedDataCache",
-				"MessageLog",
 				"Projects",
 				"SourceControl"
-			});
+			}
+		);
 
 		if (Target.bUsesSlate)
 		{
-			// NOTE: Slate is being included via SourceControl anyway
-			PrivateDependencyModuleNames.Add("Slate");
-			PrivateDependencyModuleNames.Add("SlateCore");
+			PrivateDependencyModuleNames.AddRange(
+				new string[] {
+					"MessageLog",
+					"Slate",
+					"SlateCore"
+				}
+			);
 
 			PublicDefinitions.Add("UE_VA_WITH_SLATE=1");
 		}
