Allow the VA failed payload pull dialog to be displayed if the pull occurred  when generating content browser thumbnails.

#rnx

- When generating content browser thumnails GIsRunningUnattendedScript will be set to true which prevents us from showing the VA payload pull failed dialog, so if a payload does fail during thumbnail generation the process will just terminate rather than asking the user if they would like to retry the operation.
- We no longer consider 'GIsRunningUnattendedScript' when checking if the process is interactive.
- If the dialog is being displayed from the GameThread we can force GIsRunningUnattendedScript to false safely allowing the dialog to be displayed. If we fail a payload pull on a background thread while the GameThread is generating a thumbnail we will not be able to safely change GIsRunningUnattendedScript (which does not have proper thread guarding) so we will just have to accept that the dialog will fail to show and the process will terminate.

diff --git a/Engine/Source/Developer/Virtualization/Private/VirtualizationManager.cpp b/Engine/Source/Developer/Virtualization/Private/VirtualizationManager.cpp
--- a/Engine/Source/Developer/Virtualization/Private/VirtualizationManager.cpp
+++ b/Engine/Source/Developer/Virtualization/Private/VirtualizationManager.cpp
@@ -91,6 +91,38 @@ private:
 	FCriticalSection* SyncObject;
 };
 
+/**
+ * Utility to disable 'GIsRunningUnattendedScript' for a given scope.
+ * We only change 'GIsRunningUnattendedScript' if called on the GameThread as the
+ * value is not thread safe and is generally only changed on that thread.
+ * 
+ * We need this so that we can try to force a modal dialog to be shown during
+ * the editor if a payload fails to pull as our only other option is to
+ * terminate the process.
+ */
+struct FDisableUnattendedScriptGlobal
+{
+	FDisableUnattendedScriptGlobal()
+	{
+		if (IsInGameThread())
+		{
+			bOriginalValue = GIsRunningUnattendedScript;
+			GIsRunningUnattendedScript = false;
+		}
+	}
+
+	~FDisableUnattendedScriptGlobal()
+	{
+		if (IsInGameThread())
+		{
+			GIsRunningUnattendedScript = bOriginalValue;
+		}
+	}
+
+private:
+	bool bOriginalValue = false;
+};
+
 /** 
  * Utility class to help manage pull requests. When created it will remove invalid and duplicate requests so
  * that backends do not need to worry about them.
@@ -2031,6 +2063,7 @@ FVirtualizationManager::ErrorHandlingResult FVirtualizationManager::OnPayloadPul
 
 			const FText Message = MsgBuilder.ToText();
 
+			FDisableUnattendedScriptGlobal DisableScope;
 			Result = FMessageDialog::Open(EAppMsgType::YesNo, EAppReturnType::No, Message, Title);
 		}
 		else
diff --git a/Engine/Source/Developer/Virtualization/Private/VirtualizationUtilities.cpp b/Engine/Source/Developer/Virtualization/Private/VirtualizationUtilities.cpp
--- a/Engine/Source/Developer/Virtualization/Private/VirtualizationUtilities.cpp
+++ b/Engine/Source/Developer/Virtualization/Private/VirtualizationUtilities.cpp
@@ -142,7 +142,29 @@ bool ExpandEnvironmentVariables(FStringView InputPath, FStringBuilderBase& OutEx
 
 bool IsProcessInteractive()
 {
-	return !FApp::IsUnattended() && !IsRunningCommandlet() && !GIsRunningUnattendedScript && !IS_PROGRAM;
+	if (FApp::IsUnattended())
+	{
+		return false;
+	}
+
+	if (IsRunningCommandlet())
+	{
+		return false;
+	}
+
+	// We used to check 'GIsRunningUnattendedScript' here as well but there
+	// are a number of places in the editor enabling this global during which
+	// the editor does stay interactive, such as when rendering thumbnail
+	// images for the content browser.
+	// Leaving this comment block here to show why we are not checking this
+	// value anymore.
+
+	if (IS_PROGRAM)
+	{
+		return false;
+	}
+
+	return true;
 }
 
 } // namespace UE::Virtualization::Utils
