Improve the error handling for bad 'IgnoreFile' values when defining a VA graph and improve the error reporting should a bad value slip through.

- When we first parse the value from the config file we should check to see if it contains and path separators and if so reject the backend entirely. Providing a full path rather than just the filename is an easy enough mistake to make and it is better for the user to find out sooner than later.
- The utility function ::TryCreateSubmissionSessionDirectory was extended after it was first created to also create a dummy ignore file, updated the name to reflect that it does more than just create a directory.
- We were relying ont he creation of the dummy file to also create the directory but that means it is harder for us to inform the user if an encountered problem is due to the directory creation or the file creation. Now we ensure that the directory is created first (with error handling) and only then try to create the dummy file.
- Added logging of the system error if we encounter issues to give the user a better idea what the issue was.

diff --git a/Engine/Source/Developer/Virtualization/Private/VirtualizationSourceControlBackend.cpp b/Engine/Source/Developer/Virtualization/Private/VirtualizationSourceControlBackend.cpp
--- a/Engine/Source/Developer/Virtualization/Private/VirtualizationSourceControlBackend.cpp
+++ b/Engine/Source/Developer/Virtualization/Private/VirtualizationSourceControlBackend.cpp
@@ -8,8 +8,10 @@
 #include "IO/IoHash.h"
 #include "ISourceControlModule.h"
 #include "ISourceControlProvider.h"
+#include "HAL/PlatformFileManager.h"
 #include "Interfaces/IPluginManager.h"
 #include "Logging/MessageLog.h"
+#include "Logging/StructuredLog.h"
 #include "Misc/CommandLine.h"
 #include "Misc/FileHelper.h"
 #include "Misc/Parse.h"
@@ -205,13 +207,21 @@ private:
 	return FString(TEXT("Unknown"));
 }
 
-/** Utility function to create a directory to submit payloads from. */
-[[nodiscard]] static bool TryCreateSubmissionSessionDirectory(FStringView SessionDirectoryPath, FStringView IgnoreFileName)
+/** Utility function to create a directory to submit payloads from and add the dummy P4IGNORE file */
+[[nodiscard]] static bool TrySetupSubmissionSessionDirectory(FStringView SessionDirectoryPath, FStringView IgnoreFileName)
 {
+	if (!FPlatformFileManager::Get().GetPlatformFile().CreateDirectoryTree(WriteToString<512>(SessionDirectoryPath).ToString()))
+	{
+		TStringBuilder<MAX_SPRINTF> SystemErrorMsg;
+		Utils::GetFormattedSystemError(SystemErrorMsg);
+
+		UE_LOGFMT(LogVirtualization, Error, "Failed to create '{SubmissionDir}' due to: {SysError}", SessionDirectoryPath, SystemErrorMsg);
+		return false;
+	}
+
 	// Write out an ignore file to the submission directory (will create the directory if needed)
 	{
 		TStringBuilder<260> IgnoreFilePath;
-
 		FPathViews::Append(IgnoreFilePath, SessionDirectoryPath, IgnoreFileName);
 
 		// A very basic .p4ignore file that should make sure that we are only submitting valid .upayload files.
@@ -226,6 +236,10 @@ private:
 
 		if (!FFileHelper::SaveStringToFile(FileContents, IgnoreFilePath.ToString()))
 		{
+			TStringBuilder<MAX_SPRINTF> SystemErrorMsg;
+			Utils::GetFormattedSystemError(SystemErrorMsg);
+
+			UE_LOGFMT(LogVirtualization, Error, "Failed to create '{IgnoreFile}' due to: {SysError}", IgnoreFilePath, SystemErrorMsg);
 			return false;
 		}
 	}
@@ -715,9 +729,9 @@ bool FSourceControlBackend::PushData(TArrayView<FPushRequest> Requests, EPushFla
 	TStringBuilder<260> SessionDirectory;
 	FPathViews::Append(SessionDirectory, SubmissionRootDir, SessionGuid);
 
-	if (!TryCreateSubmissionSessionDirectory(SessionDirectory, IgnoreFileName))
+	if (!TrySetupSubmissionSessionDirectory(SessionDirectory, IgnoreFileName))
 	{
-		UE_LOG(LogVirtualization, Error, TEXT("[%s] Failed to created directory '%s' to submit payloads from"), *GetDebugName(), SessionDirectory.ToString());
+		UE_LOG(LogVirtualization, Error, TEXT("[%s] Failed to setup the directory to submit payloads from"), *GetDebugName());
 		return false;
 	}
 
@@ -1118,7 +1132,17 @@ bool FSourceControlBackend::TryApplySettingsFromConfigFiles(const FString& Confi
 		// TODO: We should just extract this from the perforce environment but that requires extending
 		// the source control api.
 		// Letting the backend define the ignore filename to use is a quicker work around
-		FParse::Value(*ConfigEntry, TEXT("IgnoreFile="), IgnoreFileName);
+		if (FParse::Value(*ConfigEntry, TEXT("IgnoreFile="), IgnoreFileName))
+		{
+#if 0
+			if (IgnoreFileName.Contains(TEXT("\\")) || IgnoreFileName.Contains(TEXT("/")))
+			{
+				UE_LOG(LogVirtualization, Error, TEXT("[%s] Invalid IgnoreFile entry! This value should just be the filename and not contain any path info"), *GetDebugName());
+				return false;
+			}
+#endif
+		}
+
 		UE_LOG(LogVirtualization, Log, TEXT("[%s] Using '%s' as the p4 ignore file name"), *GetDebugName(), *IgnoreFileName);
 	}
 
diff --git a/Engine/Source/Developer/Virtualization/Private/VirtualizationSourceControlBackend.h b/Engine/Source/Developer/Virtualization/Private/VirtualizationSourceControlBackend.h
--- a/Engine/Source/Developer/Virtualization/Private/VirtualizationSourceControlBackend.h
+++ b/Engine/Source/Developer/Virtualization/Private/VirtualizationSourceControlBackend.h
@@ -73,11 +73,11 @@ class FSemaphore;
  * IgnoreFile [string]:				Sets the name of the p4 ignore file to use. When submitting payloads we 
  *									create a custom p4 ignore file to override any ignore settings for a project
  *									which allows us to submit from the saved directory which is normally prevented
- *									by the default ignore file. This value can be set to what ever is used by
- *									your perforce environment. [Default=".p4ignore.txt"]
+ *									by the default ignore file. This value should be set to the value of P4IGNORE
+									used by your environment. [Default=".p4ignore.txt"]
  * UseRetryConnectionDialog[bool]	When true a slate dialog will be shown if the initial connection to the 
  *									source control server fails allowing the user to attempt to input the correct
- *									login values. [Default=false]&
+ *									login values. [Default=false]
  * 
  * Environment Variables:
  * UE-VirtualizationWorkingDir [string]:	This can be set to a valid directory path that the backend
