Add config file options so that the virtualization system is able to retry pulling failed payloads when running in unattended mode (defaults to off)

-Some users have reported seeing their long cooks fail due to a machine having an unreliable connection or networkcard. In these cases the network outage can be swiftly fixed but if VA failed to pull a payload before then, the cook will terminate and need to be restarted which can cost a lot of time.
- It was requested that we add an optional way to have the system retry payload pulling when running in unattended mode but also to wait for X time (usually many minutes) before trying again. If it is likely that the connection will be restored within those few minutes then waiting will be much less costly than restarting a cook.

- That payloads can be pulled on many threads at the same time makes the logic a little tricky, so rather than counting how many payloads have failed vs the retry counter we count how many times we've logged a message to the user as this logging is protected by a critical section and acts as a way to "group" together failed pulls that occur around about the same time. We then reset this counter to 0 if we detect a successful pull.
- It is possible that a pull fails because the payload is missing, in which case this logic will probably cause the counter to reset frequently and the error to not become fatal for quite some time (possibly until the cook has almost finished) but it is quite unlikly to occur and due to this I have favored erring towards simple code rather than trying to track individual payload failures vs grouped failures vs successful pulls.

- Note: That when backends fail to pull payloads they generally log errors, which will eventually cause most of our processes to return non zero to indicate failure. VA should not log errors while we are inside of a retry loop and only print out errors when we detect a problem that we cannot solve to avoid this. This is being addressed as it's own work item.

diff --git a/Engine/Config/BaseEngine.ini b/Engine/Config/BaseEngine.ini
--- a/Engine/Config/BaseEngine.ini
+++ b/Engine/Config/BaseEngine.ini
@@ -1477,6 +1477,8 @@ UseLegacyErrorHandling=False
 ForceCachingOnPull=False
 PullErrorAdditionalMsg=""
 ConnectionHelpUrl=""
+UnattendedRetryCount=0
+UnattendedRetryTimer=0
 
 [/Script/Virtualization.VirtualizationFilterSettings]
 !ExcludePackagePaths=ClearArray
diff --git a/Engine/Source/Developer/Virtualization/Private/VirtualizationManager.cpp b/Engine/Source/Developer/Virtualization/Private/VirtualizationManager.cpp
--- a/Engine/Source/Developer/Virtualization/Private/VirtualizationManager.cpp
+++ b/Engine/Source/Developer/Virtualization/Private/VirtualizationManager.cpp
@@ -853,7 +853,12 @@ bool FVirtualizationManager::PullData(TArrayView<FPullRequest> Requests)
 		{
 			checkf(Request.GetIdentifier() == Request.GetPayload().GetRawHash(), TEXT("Invalid payload for '%s'"), *LexToString(Request.GetIdentifier()));
 		}
-	}	
+	}
+
+	if (bSuccess)
+	{
+		UnattendedFailureMsgCount = 0;
+	}
 
 	return bSuccess;
 }
@@ -1214,6 +1219,20 @@ void FVirtualizationManager::ApplySettingsFromConfigFiles(const FConfigFile& Con
 		UE_LOG(LogVirtualization, Display, TEXT("\tForceCachingOnPull : %s"), bForceCachingOnPull ? TEXT("true") : TEXT("false"));
 	}
 
+	{
+		ConfigFile.GetInt(ConfigSection, TEXT("UnattendedRetryTimer"), UnattendedRetryTimer);
+		ConfigFile.GetInt(ConfigSection, TEXT("UnattendedRetryCount"), UnattendedRetryCount);
+		
+		if (ShouldRetryWhenUnattended())
+		{
+			UE_LOG(LogVirtualization, Display, TEXT("\tUnattendedRetryTimer : %d retries every %d(s)"), UnattendedRetryCount, UnattendedRetryTimer);
+		}
+		else
+		{
+			UE_LOG(LogVirtualization, Display, TEXT("\tUnattendedRetries : disabled"));
+		}
+	}
+
 	// Deprecated
 	{
 		bool bDummyValue = true;
@@ -2000,9 +2019,9 @@ void FVirtualizationManager::PullDataFromBackend(IVirtualizationBackend& Backend
 	COOK_STAT(FCookStats::CallStats & Stats = Profiling::GetPullStats(Backend));
 	COOK_STAT(FCookStats::FScopedStatsCounter Timer(Stats));
 	COOK_STAT(Timer.TrackCyclesOnly());
-	
+
 	Backend.PullData(Requests, IVirtualizationBackend::EPullFlags::None, OutErrors);
-	
+
 #if ENABLE_COOK_STATS
 	const bool bIsInGameThread = IsInGameThread();
 
@@ -2028,7 +2047,7 @@ FVirtualizationManager::ErrorHandlingResult FVirtualizationManager::OnPayloadPul
 
 	for (const FIoHash& FailedPayload : Requests.GetFailedPayloads())
 	{
-		UE_LOG(LogVirtualization, Error, TEXT("Payload '%s' failed to be pulled from any backend'"), *LexToString(FailedPayload));
+		UE_LOG(LogVirtualization, Error, TEXT("Failed to pull payload '%s'"), *LexToString(FailedPayload));
 	}
 
 	static FCriticalSection CriticalSection;
@@ -2069,7 +2088,27 @@ FVirtualizationManager::ErrorHandlingResult FVirtualizationManager::OnPayloadPul
 		else
 		{
 			const FText Message = MsgBuilder.ToText();
-			UE_LOG(LogVirtualization, Error, TEXT("%s"), *Message.ToString());
+
+			if (!ShouldRetryWhenUnattended() || UnattendedFailureMsgCount >= UnattendedRetryCount)
+			{
+				UE_LOG(LogVirtualization, Error, TEXT("%s"), *Message.ToString());
+			}
+			else
+			{
+				// Only log as a warning while we are retrying so that if we do recover out logging will
+				// not cause the running process to be considered a failure due to an error.
+				UE_LOG(LogVirtualization, Warning, TEXT("Failed to pull payload(s) on attempt (%d/%d)"), UnattendedFailureMsgCount.load(), UnattendedRetryCount);
+
+				Result = EAppReturnType::Yes; // Attempt to retry the failed pull
+
+				if (UnattendedRetryTimer > 0)
+				{
+					UE_LOG(LogVirtualization, Warning, TEXT("Waiting for %d seconds before trying to pull again..."), UnattendedRetryTimer);
+					FPlatformProcess::SleepNoStats(static_cast<float>(UnattendedRetryTimer));
+				}
+
+				UnattendedFailureMsgCount++;
+			}
 		}
 
 		if (Result == EAppReturnType::No)
@@ -2219,6 +2258,11 @@ bool FVirtualizationManager::ShouldVirtualizeAsDefault() const
 	}
 }
 
+bool FVirtualizationManager::ShouldRetryWhenUnattended() const
+{
+	return UnattendedRetryCount > 0;
+}
+
 void FVirtualizationManager::BroadcastEvent(TConstArrayView<FPullRequest> Requests, ENotification Event)
 {
 	for (const FPullRequest& Request : Requests)
diff --git a/Engine/Source/Developer/Virtualization/Private/VirtualizationManager.h b/Engine/Source/Developer/Virtualization/Private/VirtualizationManager.h
--- a/Engine/Source/Developer/Virtualization/Private/VirtualizationManager.h
+++ b/Engine/Source/Developer/Virtualization/Private/VirtualizationManager.h
@@ -135,6 +135,17 @@ struct FAnalyticsEventAttribute;
  *												before it was pulled from a backend later in the hierarchy. Can be used to try and skip
  *												expensive existence checks, or if a backend is in a bad state where it believes it has the payload
  *												but is unable to actually return the data. [Default=false]
+ * UnattendedRetryCount [int32]:				How many times the process should retry pulling payloads after a failure is encountered if the
+												process is unattended. Usually when a payload pull fails we ask the user to try and fix the issue
+												and retry, but in unattended mode we just log an error and terminate the process. In some cases
+												such as build machines with unreliable internet it is possible that the process could recover in
+												which case setting this value might help. Zero or negative values will disable the system. 
+												Note: If many pulls are occurring at the same time on many threads a very short network outage might
+												spawn many errors in which case we try to group these errors into a single 'try' so 32 errors on 32
+												threads would not immediately blow past a retry count of 30 for example. [Default=0]
+ * UnattendedRetryTimer [int32]					If 'UnattendedRetryCount' is set to a positive value then this value sets how long (in seconds)
+ *												the process should wait after a failure is encountered before retrying the pull. Depending on the
+												likely cause of the failure you may want to set this value to several minutes. [Default=0]
  */
 
 namespace UE::Virtualization
@@ -274,6 +285,9 @@ private:
 	/** Determines if the default filtering behavior is to virtualize a payload or not */
 	bool ShouldVirtualizeAsDefault() const;
 
+	/** Returns if the process will attempt to retry a failed pull when the process is unattended mode */
+	bool ShouldRetryWhenUnattended() const;
+
 	void BroadcastEvent(TConstArrayView<FPullRequest> Ids, ENotification Event);
 	
 private:
@@ -331,6 +345,13 @@ private:
 
 	/** Optional url used to augment connection failure error messages */
 	static FString ConnectionHelpUrl;
+
+	/** The number of times to retry pulling when errors are encountered in an unattended process, values <0 disable the system */
+	int32 UnattendedRetryCount = 0;
+
+	/** The how long (in seconds) to wait after payload pulling errors before retrying. Does nothing if 'UnattendedRetryCount' is disabled */
+	int32 UnattendedRetryTimer = 0;
+
 private:
 
 	/** The name of the current project */
@@ -360,6 +381,9 @@ private:
 	/** Our notification Event */
 	FOnNotification NotificationEvent;
 
+	/** Track how many times we've displayed a message about failed payload pulls since the last successful pull (only used in unattended mode)*/
+	mutable std::atomic<int32> UnattendedFailureMsgCount = 0;
+
 	// Members after this point at used for debugging operations only!
 
 	struct FDebugValues
