Add analytics for asset virtualization reliability events.

#rnx

- Add a new event (GetAnalyticsRecordEvent) to the UE::Virtualization namespace. Different telemetry systems can hook into this and respond when the virtualization system tries to send a new analytics event.
-- This even accepts an enum class of flags to provide better context for the event. At the moment the only valid flag is to request that the telemetry system flushes the event to make sure that it is sent before returning control to the event. In this way we can know that events are safely recorded before terminating the current process.
- To begin with we are only adding failure events and not recording any successes
-- Send an analytics even when the initial connection to perforce fails, the event should note if the user opted to retry their connection (after entering new details) or if they wanted to skip the check and load the editor anyway.
-- Send an analytics event when a payload fails to pull, the event should note if the user opted to retry to pull the payload or quit the process.

diff --git a/Engine/Plugins/Experimental/StudioTelemetry/Source/StudioTelemetry/Private/StudioTelemetryEditor.cpp b/Engine/Plugins/Experimental/StudioTelemetry/Source/StudioTelemetry/Private/StudioTelemetryEditor.cpp
--- a/Engine/Plugins/Experimental/StudioTelemetry/Source/StudioTelemetry/Private/StudioTelemetryEditor.cpp
+++ b/Engine/Plugins/Experimental/StudioTelemetry/Source/StudioTelemetry/Private/StudioTelemetryEditor.cpp
@@ -923,6 +923,21 @@ void FStudioTelemetryEditor::Initialize()
 			});
 		});
 	}
+
+	{
+		UE::Virtualization::GetAnalyticsRecordEvent().AddLambda([](const FString& EventName, const TArray<FAnalyticsEventAttribute>& Attributes, UE::Virtualization::EAnalyticsFlags Flags)
+			{
+				FStudioTelemetry::Get().RecordEvent(EventName, Attributes);
+
+				if (EnumHasAllFlags(Flags, UE::Virtualization::EAnalyticsFlags::Flush))
+				{
+					if (TSharedPtr<IAnalyticsProvider> Provider = FStudioTelemetry::Get().GetProvider().Pin())
+					{
+						Provider->BlockUntilFlushed(60.0f);
+					}
+				}
+			});
+	}
 }
 
 void FStudioTelemetryEditor::Shutdown()
diff --git a/Engine/Source/Developer/Virtualization/Private/SVirtualizationRevisionControlConnectionDialog.cpp b/Engine/Source/Developer/Virtualization/Private/SVirtualizationRevisionControlConnectionDialog.cpp
--- a/Engine/Source/Developer/Virtualization/Private/SVirtualizationRevisionControlConnectionDialog.cpp
+++ b/Engine/Source/Developer/Virtualization/Private/SVirtualizationRevisionControlConnectionDialog.cpp
@@ -80,12 +80,16 @@ SRevisionControlConnectionDialog::FResult SRevisionControlConnectionDialog::RunD
 	TSharedPtr<SWindow> ParentWindow = FSlateApplication::Get().GetActiveTopLevelWindow();
 	FSlateApplication::Get().AddModalWindow(DialogWindow.ToSharedRef(), ParentWindow);
 
-	if (DialogWidget->GetResult() == SRevisionControlConnectionDialog::EResult::Retry)
+	switch (DialogWidget->GetResult())
 	{
-		return FResult(DialogWidget->GetPort(), DialogWidget->GetUserName());
+		case EResult::Retry:
+			return FResult(DialogWidget->GetPort(), DialogWidget->GetUserName());
+		case EResult::Skip:
+			return FResult();
+		default:
+			checkNoEntry();
+			return FResult();
 	}
-
-	return FResult();
 }
 
 void SRevisionControlConnectionDialog::Construct(const FArguments& InArgs, FStringView RevisionControlName, FStringView InConfigSectionName, FStringView CurrentPort, FStringView CurrentUsername, const FText& ErrorMessage)
diff --git a/Engine/Source/Developer/Virtualization/Private/VirtualizationManager.cpp b/Engine/Source/Developer/Virtualization/Private/VirtualizationManager.cpp
--- a/Engine/Source/Developer/Virtualization/Private/VirtualizationManager.cpp
+++ b/Engine/Source/Developer/Virtualization/Private/VirtualizationManager.cpp
@@ -2007,10 +2007,13 @@ void FVirtualizationManager::PullDataFromAllBackends(TArrayView<FPullRequest> Re
 		{
 			return; // All payloads pulled
 		}
-		else if (OnPayloadPullError(RequestsCollection, BackendErrors) != ErrorHandlingResult::Retry)
+		else if (OnPayloadPullError(RequestsCollection, BackendErrors) == ErrorHandlingResult::AcceptFailedPayloads)
 		{
-			return; // Some payloads failed to pull
+			return; // Some payloads failed to pull but we will let the calling code deal with that.
+					// This path should only be taken if legacy error handling is enabled.
 		}
+
+		// The user opted to retry pulling the payloads that failed.
 	}
 }
 
@@ -2111,6 +2114,16 @@ FVirtualizationManager::ErrorHandlingResult FVirtualizationManager::OnPayloadPul
 			}
 		}
 
+		{
+			TArray<FAnalyticsEventAttribute> Attributes;
+			Attributes.Add({ TEXT("UserSelection"), Result == EAppReturnType::No ? TEXT("Quit") : TEXT("Retry")});
+
+			// Need to flush if we are going to quit to make sure that the analytics payloads are sent properly.
+			const EAnalyticsFlags Flags = Result == EAppReturnType::No ? EAnalyticsFlags::Flush : EAnalyticsFlags::None;
+
+			GetAnalyticsRecordEvent().Broadcast(TEXT("Editor.VA.PayloadPullError"), Attributes, Flags);
+		}
+
 		if (Result == EAppReturnType::No)
 		{
 			// Reporting a fatal error will print the callstack and initiate the crash handling system to be treated like a bug
diff --git a/Engine/Source/Developer/Virtualization/Private/VirtualizationSourceControlBackend.cpp b/Engine/Source/Developer/Virtualization/Private/VirtualizationSourceControlBackend.cpp
--- a/Engine/Source/Developer/Virtualization/Private/VirtualizationSourceControlBackend.cpp
+++ b/Engine/Source/Developer/Virtualization/Private/VirtualizationSourceControlBackend.cpp
@@ -2,6 +2,7 @@
 
 #include "VirtualizationSourceControlBackend.h"
 
+#include "AnalyticsEventAttribute.h"
 #include "Containers/Ticker.h"
 #include "HAL/Event.h"
 #include "HAL/FileManager.h"
@@ -22,10 +23,13 @@
 #include "SourceControlHelpers.h"
 #include "SourceControlInitSettings.h"
 #include "SourceControlOperations.h"
+#include "Virtualization/VirtualizationSystem.h"
 #include "VirtualizationManager.h"
 #include "VirtualizationSourceControlUtilities.h"
 #include "VirtualizationUtilities.h"
 
+#include "HAL/PlatformFileManager.h"
+
 // When the SourceControl module (or at least the perforce source control module) is thread safe we
 // can enable this and stop using the hacky work around 'TryToDownloadFileFromBackgroundThread'
 #define IS_SOURCE_CONTROL_THREAD_SAFE 0
@@ -389,8 +393,19 @@ IVirtualizationBackend::EConnectionStatus FSourceControlBackend::OnConnect()
 																													UserName,
 																													ErrorMessage);
 
+			{
+				TArray<FAnalyticsEventAttribute> Attributes;
+				Attributes.Add({ TEXT("UserSelection"), DialogResult.bShouldRetry ? TEXT("Retry") : TEXT("Skip") });
+
+				// Need to flush if we are going to quit to make sure that the analytics payloads are sent properly.
+				const EAnalyticsFlags Flags = DialogResult.bShouldRetry ? EAnalyticsFlags::None : EAnalyticsFlags::Flush;
+
+				GetAnalyticsRecordEvent().Broadcast(TEXT("Editor.VA.SourceControlLoginFail"), Attributes, EAnalyticsFlags::None);
+			}
+			
 			if (!DialogResult.bShouldRetry)
 			{
+				// User has opted to continue without a valid source control connection
 				OnConnectionError(ErrorMessage);
 				return IVirtualizationBackend::EConnectionStatus::Error;
 			}
diff --git a/Engine/Source/Runtime/Core/Private/Virtualization/VirtualizationSystem.cpp b/Engine/Source/Runtime/Core/Private/Virtualization/VirtualizationSystem.cpp
--- a/Engine/Source/Runtime/Core/Private/Virtualization/VirtualizationSystem.cpp
+++ b/Engine/Source/Runtime/Core/Private/Virtualization/VirtualizationSystem.cpp
@@ -363,6 +363,11 @@ void Shutdown()
 	UE_LOG(LogVirtualization, Verbose, TEXT("UE::Virtualization was shutdown"));
 }
 
+FAnalyticsRecordEvent& GetAnalyticsRecordEvent()
+{
+	static FAnalyticsRecordEvent Event;
+	return Event;
+}
 
 bool IVirtualizationSystem::IsInitialized()
 {
diff --git a/Engine/Source/Runtime/Core/Public/Virtualization/VirtualizationSystem.h b/Engine/Source/Runtime/Core/Public/Virtualization/VirtualizationSystem.h
--- a/Engine/Source/Runtime/Core/Public/Virtualization/VirtualizationSystem.h
+++ b/Engine/Source/Runtime/Core/Public/Virtualization/VirtualizationSystem.h
@@ -529,6 +529,23 @@ CORE_API bool ShouldInitializePreSlate();
  */
 CORE_API void Shutdown();
 
+/** Used to customize the analytics events being sent via the virtualization system */
+enum class EAnalyticsFlags : uint32
+{
+	None = 0,
+	/** The event needs to be completely handled/flushed by the listening systems before returning.
+		Commonly used when sending events right before terminating the process */
+	Flush = 1 << 0
+};
+
+using FAnalyticsRecordEvent = TTSMulticastDelegate<void(const FString&, const TArray<FAnalyticsEventAttribute>&, EAnalyticsFlags)>;
+
+/** 
+ * This event is used by the virtualization system to send analytics events. Telemetry systems that wish to process asset virtualization
+ * events should register to this. Note that these events could be raised on different threads so the listener will need to be threadsafe.
+ */
+CORE_API FAnalyticsRecordEvent& GetAnalyticsRecordEvent();
+
 /** 
  * The base interface for the virtualization system. An Epic maintained version can be found in the Virtualization module.
  * To implement your own, simply derived from this interface and then use the
