Only log warnings for the specific payload(s) that failed to cache rather than for all payloads in the batch.

#rnx

diff --git a/Engine/Source/Developer/Virtualization/Private/VirtualizationManager.cpp b/Engine/Source/Developer/Virtualization/Private/VirtualizationManager.cpp
--- a/Engine/Source/Developer/Virtualization/Private/VirtualizationManager.cpp
+++ b/Engine/Source/Developer/Virtualization/Private/VirtualizationManager.cpp
@@ -1853,7 +1853,9 @@ void FVirtualizationManager::CachePayloads(TArrayView<FPushRequest> Requests, co
 		{
 			for (const FPushRequest& Request : Requests)
 			{
-				UE_LOG(LogVirtualization, Warning, TEXT("Failed to cache payload '%s' to backend '%s'"), *LexToString(Request.GetIdentifier()), *BackendToCache->GetDebugName());
+				UE_CLOG(Request.GetResult().WasError(), LogVirtualization, Warning, TEXT("Failed to cache payload '%s' to backend '%s'"),
+					*LexToString(Request.GetIdentifier()),
+					*BackendToCache->GetDebugName());
 			}
 		}
 
diff --git a/Engine/Source/Runtime/Core/Public/Virtualization/VirtualizationSystem.h b/Engine/Source/Runtime/Core/Public/Virtualization/VirtualizationSystem.h
--- a/Engine/Source/Runtime/Core/Public/Virtualization/VirtualizationSystem.h
+++ b/Engine/Source/Runtime/Core/Public/Virtualization/VirtualizationSystem.h
@@ -158,6 +158,12 @@ struct FPushResult
 		return Status == EStatus::Pushed;
 	}
 
+	/** Returns true if the payload experienced an error during upload */
+	bool WasError() const
+	{
+		return Status < EStatus::Pending;
+	}
+
 	/** Returns true if the payload is stored in the target backend(s) after the operation completed */
 	bool IsVirtualized() const
 	{
