Add verbose logging when caching payloads.

#rnx

- Add a verbose log message each time a payload is actually cached to a backend detailing the backend and the payload identifier.
-- Should only log if the payload actually uploads to the given backend, if it is rejected because the payload already exists in the backend we will not log anything.
- In order for this to work we need to reset the status of each FPushRequest being cached as we might have more than one caching backend and we don't want to includes pushes from one backend to another.

diff --git a/Engine/Source/Developer/Virtualization/Private/VirtualizationManager.cpp b/Engine/Source/Developer/Virtualization/Private/VirtualizationManager.cpp
--- a/Engine/Source/Developer/Virtualization/Private/VirtualizationManager.cpp
+++ b/Engine/Source/Developer/Virtualization/Private/VirtualizationManager.cpp
@@ -1847,6 +1847,12 @@ void FVirtualizationManager::CachePayloads(TArrayView<FPushRequest> Requests, co
 			return;
 		}
 
+		// Reset any previous caching results
+		for (FPushRequest& Request : Requests)
+		{
+			Request.ResetResult();
+		}	
+
 		const bool bResult = TryCacheDataToBackend(*BackendToCache, Requests, Flags);
 
 		if (!bResult)
@@ -1897,6 +1903,8 @@ bool FVirtualizationManager::TryCacheDataToBackend(IVirtualizationBackend& Backe
 				Stats.Accumulate(FCookStats::CallStats::EHitOrMiss::Hit, FCookStats::CallStats::EStatType::Bytes, Request.GetPayload().GetCompressedSize(), bIsInGameThread);
 
 				bWasDataPushed = true;
+
+				UE_LOG(LogVirtualization, Verbose, TEXT("[%s] Cached payload '%s'"), *Backend.GetDebugName(), *LexToString(Request.GetIdentifier()));
 			}
 		}
 
