Clean up FVirtualizationManager::GatherAnalytics

#rnx

- When this code was moved from elsewhere to FVirtualizationManager we neglected to remove the calls to IVirtualizationSystem. The code was working just fine but on first glance it is more complicated that it needed to be as it looks like a static method (hence needing to find the active IVirtualizationSystem pointer) but it is not.
- Reserve the space we need in the attributes array.

diff --git a/Engine/Source/Developer/Virtualization/Private/VirtualizationManager.cpp b/Engine/Source/Developer/Virtualization/Private/VirtualizationManager.cpp
--- a/Engine/Source/Developer/Virtualization/Private/VirtualizationManager.cpp
+++ b/Engine/Source/Developer/Virtualization/Private/VirtualizationManager.cpp
@@ -2296,66 +2296,60 @@ void FVirtualizationManager::BroadcastEvent(TConstArrayView<FPullRequest> Reques
 
 void FVirtualizationManager::GatherAnalytics(TArray<FAnalyticsEventAttribute>& Attributes) const
 {
-	using namespace UE::Virtualization;
+	Attributes.Reserve(Attributes.Num() + 10);
 
-	// Grab the Virtualization stats
-	if (IVirtualizationSystem::IsInitialized())
-	{
-		IVirtualizationSystem& System = IVirtualizationSystem::Get();
-
-		FPayloadActivityInfo PayloadActivityInfo = System.GetAccumualtedPayloadActivityInfo();
+	FPayloadActivityInfo PayloadActivityInfo = GetAccumualtedPayloadActivityInfo();
 
-		const FString BaseName = TEXT("Virtualization");
+	const FString BaseName = TEXT("Virtualization");
 
-		{
-			FString AttrName = BaseName + TEXT("_Enabled");
-			Attributes.Emplace(MoveTemp(AttrName), System.IsEnabled());
-		}
+	{
+		FString AttrName = BaseName + TEXT("_Enabled");
+		Attributes.Emplace(MoveTemp(AttrName), IsEnabled());
+	}
 
-		{
-			FString AttrName = BaseName + TEXT("_Cache_TimeSpent");
-			Attributes.Emplace(MoveTemp(AttrName), (double)PayloadActivityInfo.Cache.CyclesSpent * FPlatformTime::GetSecondsPerCycle());
-		}
+	{
+		FString AttrName = BaseName + TEXT("_Cache_TimeSpent");
+		Attributes.Emplace(MoveTemp(AttrName), (double)PayloadActivityInfo.Cache.CyclesSpent * FPlatformTime::GetSecondsPerCycle());
+	}
 
-		{
-			FString AttrName = BaseName + TEXT("_Cache_PayloadCount");
-			Attributes.Emplace(MoveTemp(AttrName), (double)PayloadActivityInfo.Cache.PayloadCount);
-		}
+	{
+		FString AttrName = BaseName + TEXT("_Cache_PayloadCount");
+		Attributes.Emplace(MoveTemp(AttrName), (double)PayloadActivityInfo.Cache.PayloadCount);
+	}
 
-		{
-			FString AttrName = BaseName + TEXT("_Cache_TotalBytes");
-			Attributes.Emplace(MoveTemp(AttrName), (double)PayloadActivityInfo.Cache.TotalBytes);
-		}
+	{
+		FString AttrName = BaseName + TEXT("_Cache_TotalBytes");
+		Attributes.Emplace(MoveTemp(AttrName), (double)PayloadActivityInfo.Cache.TotalBytes);
+	}
 
-		{
-			FString AttrName = BaseName + TEXT("_Push_TimeSpent");
-			Attributes.Emplace(MoveTemp(AttrName), (double)PayloadActivityInfo.Push.CyclesSpent * FPlatformTime::GetSecondsPerCycle());
-		}
+	{
+		FString AttrName = BaseName + TEXT("_Push_TimeSpent");
+		Attributes.Emplace(MoveTemp(AttrName), (double)PayloadActivityInfo.Push.CyclesSpent * FPlatformTime::GetSecondsPerCycle());
+	}
 
-		{
-			FString AttrName = BaseName + TEXT("_Push_PayloadCount");
-			Attributes.Emplace(MoveTemp(AttrName), (double)PayloadActivityInfo.Push.PayloadCount);
-		}
+	{
+		FString AttrName = BaseName + TEXT("_Push_PayloadCount");
+		Attributes.Emplace(MoveTemp(AttrName), (double)PayloadActivityInfo.Push.PayloadCount);
+	}
 
-		{
-			FString AttrName = BaseName + TEXT("_Push_TotalBytes");
-			Attributes.Emplace(MoveTemp(AttrName), (double)PayloadActivityInfo.Push.TotalBytes);
-		}
+	{
+		FString AttrName = BaseName + TEXT("_Push_TotalBytes");
+		Attributes.Emplace(MoveTemp(AttrName), (double)PayloadActivityInfo.Push.TotalBytes);
+	}
 
-		{
-			FString AttrName = BaseName + TEXT("_Pull_TimeSpent");
-			Attributes.Emplace(MoveTemp(AttrName), (double)PayloadActivityInfo.Pull.CyclesSpent * FPlatformTime::GetSecondsPerCycle());
-		}
+	{
+		FString AttrName = BaseName + TEXT("_Pull_TimeSpent");
+		Attributes.Emplace(MoveTemp(AttrName), (double)PayloadActivityInfo.Pull.CyclesSpent * FPlatformTime::GetSecondsPerCycle());
+	}
 
-		{
-			FString AttrName = BaseName + TEXT("_Pull_PayloadCount");
-			Attributes.Emplace(MoveTemp(AttrName), (double)PayloadActivityInfo.Pull.PayloadCount);
-		}
+	{
+		FString AttrName = BaseName + TEXT("_Pull_PayloadCount");
+		Attributes.Emplace(MoveTemp(AttrName), (double)PayloadActivityInfo.Pull.PayloadCount);
+	}
 
-		{
-			FString AttrName = BaseName + TEXT("_Pull_TotalBytes");
-			Attributes.Emplace(MoveTemp(AttrName), (double)PayloadActivityInfo.Pull.TotalBytes);
-		}
+	{
+		FString AttrName = BaseName + TEXT("_Pull_TotalBytes");
+		Attributes.Emplace(MoveTemp(AttrName), (double)PayloadActivityInfo.Pull.TotalBytes);
 	}
 }
 
