Small clean up of the virtualization stats logging

#rnx

- Note that these stats are only logged after selected commandlets have run (cooking, precaching etc)
- Replaces the data transfer rate statistic with an average request time to reflect recent changes to the statistics panel in the editor.
- Compacted the table to remove some unnecessary padding.

diff --git a/Engine/Source/Developer/Virtualization/Private/VirtualizationManager.cpp b/Engine/Source/Developer/Virtualization/Private/VirtualizationManager.cpp
--- a/Engine/Source/Developer/Virtualization/Private/VirtualizationManager.cpp
+++ b/Engine/Source/Developer/Virtualization/Private/VirtualizationManager.cpp
@@ -488,7 +488,7 @@ namespace Profiling
 	{
 		UE_LOG(LogVirtualization, Display, TEXT(""));
 		UE_LOG(LogVirtualization, Display, TEXT("Virtualization ProfileData"));
-		UE_LOG(LogVirtualization, Display, TEXT("=================================================================================================="));
+		UE_LOG(LogVirtualization, Display, TEXT("============================================================================================="));
 
 		if (!HasProfilingData())
 		{
@@ -500,25 +500,25 @@ namespace Profiling
 			{
 				if (HasProfilingData(Stats))
 				{
-					UE_LOG(LogVirtualization, Display, TEXT("%-40s|%10s|%17s|%12s|%14s|"), Name, TEXT("TotalCount"), TEXT("TotalSize (MB)"), TEXT("TotalTime(s)"), TEXT("DataRate(MB/S)"));
-					UE_LOG(LogVirtualization, Display, TEXT("----------------------------------------|----------|-----------------|------------|--------------|"));
+					UE_LOG(LogVirtualization, Display, TEXT("%-40s|%10s|%15s|%12s|%11s|"), Name, TEXT("TotalCount"), TEXT("TotalSize (MiB)"), TEXT("TotalTime(s)"), TEXT("AvgTime(ms)"));
+					UE_LOG(LogVirtualization, Display, TEXT("----------------------------------------|----------|---------------|------------|-----------|"));
 
 					for (const auto& Iterator : Stats)
 					{
 						const int64 Count = Iterator.Value.GetAccumulatedValueAnyThread(FCookStats::CallStats::EHitOrMiss::Hit, FCookStats::CallStats::EStatType::Counter);
-						const double Time = (double)Iterator.Value.GetAccumulatedValueAnyThread(FCookStats::CallStats::EHitOrMiss::Hit, FCookStats::CallStats::EStatType::Cycles) * FPlatformTime::GetSecondsPerCycle();
+						const double TotalTime = (double)Iterator.Value.GetAccumulatedValueAnyThread(FCookStats::CallStats::EHitOrMiss::Hit, FCookStats::CallStats::EStatType::Cycles) * FPlatformTime::GetSecondsPerCycle();
 						const double DataSizeMB = (double)Iterator.Value.GetAccumulatedValueAnyThread(FCookStats::CallStats::EHitOrMiss::Hit, FCookStats::CallStats::EStatType::Bytes) / (1024.0f * 1024.0f);
-						const double MBps = Time != 0.0 ? (DataSizeMB / Time) : 0.0;
+						const double AvgTime = Count > 0 ? (TotalTime * 1000.0) / static_cast<double>(Count) : 0.0;
 
-						UE_LOG(LogVirtualization, Display, TEXT("%-40.40s|%10lld|%17.1f|%12.3f|%14.3f|"),
+						UE_LOG(LogVirtualization, Display, TEXT("%-40.40s|%10lld|%15.1lf|%12.1lf|%11.0lf|"),
 							*Iterator.Key,
 							Count,
 							DataSizeMB,
-							Time,
-							MBps);
+							TotalTime,
+							AvgTime);
 					}
 
-					UE_LOG(LogVirtualization, Display, TEXT("=================================================================================================="));
+					UE_LOG(LogVirtualization, Display, TEXT("============================================================================================="));
 				}
 			};
 
