Added new methods to IVirtualizationSystem for accessing statistics and deprecated the old methods.

#rnx

- The previous functor based system was difficult to update as changes to the data being passed to the functor could break existing code, so ::GetAccumualtedPayloadActivityInfo has been replaced with ::GetBackendStatistics, a version that returns a struct for each backend in the system which contains the relevant statistics for that backend. Adding new data to a struct is much easier.
-- The new data added per backend is the storage type so that the caller can tell if the backend was for cached storage or persistent storage.
- To fit into this naming scheme ::GetAccumualtedPayloadActivityInfo has been deprecated in favor of ::GetSystemStatistics which is simply a rename.
- The deprecated methods should continue to work as before.
- Existing engine code had been updated to use the new methods.
- Removed a typo from the copyright boilerplate at the begining of VirtualizationSystem.h

diff --git a/Engine/Source/Developer/Virtualization/Private/VirtualizationManager.cpp b/Engine/Source/Developer/Virtualization/Private/VirtualizationManager.cpp
--- a/Engine/Source/Developer/Virtualization/Private/VirtualizationManager.cpp
+++ b/Engine/Source/Developer/Virtualization/Private/VirtualizationManager.cpp
@@ -984,7 +984,7 @@ void FVirtualizationManager::DumpStats() const
 #endif // ENABLE_COOK_STATS
 }
 
-FPayloadActivityInfo FVirtualizationManager::GetAccumualtedPayloadActivityInfo() const
+FPayloadActivityInfo FVirtualizationManager::GetSystemStatistics() const
 {
 	FPayloadActivityInfo Info;
 
@@ -1014,14 +1014,24 @@ FPayloadActivityInfo FVirtualizationManager::GetAccumualtedPayloadActivityInfo()
 	return Info;
 }
 
-void FVirtualizationManager::GetPayloadActivityInfo( GetPayloadActivityInfoFuncRef GetPayloadFunc ) const
+TArray<FBackendStats> FVirtualizationManager::GetBackendStatistics() const
 {
-	FPayloadActivityInfo Info;
-
+	TArray<FBackendStats> Stats;
+	
 #if ENABLE_COOK_STATS
 
-	for (const auto& Backend : AllBackends)
+	Stats.Reserve(AllBackends.Num());
+
+	for (const TUniquePtr<IVirtualizationBackend>& Backend : AllBackends)
 	{
+		FBackendStats& BackendStats = Stats.AddDefaulted_GetRef();
+
+		BackendStats.DebugName = Backend->GetDebugName();
+		BackendStats.ConfigName = Backend->GetConfigName();
+		BackendStats.Type = PersistentStorageBackends.Contains(Backend.Get()) ? EStorageType::Persistent : EStorageType::Cache;
+
+		FPayloadActivityInfo& Info = BackendStats.PayloadActivity;
+
 		const FCookStats::CallStats& CacheStats = Profiling::GetCacheStats(*Backend);
 
 		Info.Cache.PayloadCount = CacheStats.GetAccumulatedValueAnyThread(FCookStats::CallStats::EHitOrMiss::Hit, FCookStats::CallStats::EStatType::Counter);
@@ -1039,13 +1049,10 @@ void FVirtualizationManager::GetPayloadActivityInfo( GetPayloadActivityInfoFuncR
 		Info.Pull.PayloadCount = PullStats.GetAccumulatedValueAnyThread(FCookStats::CallStats::EHitOrMiss::Hit, FCookStats::CallStats::EStatType::Counter);
 		Info.Pull.TotalBytes = PullStats.GetAccumulatedValueAnyThread(FCookStats::CallStats::EHitOrMiss::Hit, FCookStats::CallStats::EStatType::Bytes);
 		Info.Pull.CyclesSpent = PullStats.GetAccumulatedValueAnyThread(FCookStats::CallStats::EHitOrMiss::Hit, FCookStats::CallStats::EStatType::Cycles);
-
-		GetPayloadFunc(Backend->GetDebugName(), Backend->GetConfigName(), Info);
-
 	}
 #endif // ENABLE_COOK_STATS
 
-	
+	return Stats;
 }
 
 void FVirtualizationManager::ApplySettingsFromConfigFiles(const FConfigFile& ConfigFile)
@@ -2298,7 +2305,7 @@ void FVirtualizationManager::GatherAnalytics(TArray<FAnalyticsEventAttribute>& A
 {
 	Attributes.Reserve(Attributes.Num() + 10);
 
-	FPayloadActivityInfo PayloadActivityInfo = GetAccumualtedPayloadActivityInfo();
+	FPayloadActivityInfo PayloadActivityInfo = GetSystemStatistics();
 
 	const FString BaseName = TEXT("Virtualization");
 
diff --git a/Engine/Source/Developer/Virtualization/Private/VirtualizationManager.h b/Engine/Source/Developer/Virtualization/Private/VirtualizationManager.h
--- a/Engine/Source/Developer/Virtualization/Private/VirtualizationManager.h
+++ b/Engine/Source/Developer/Virtualization/Private/VirtualizationManager.h
@@ -199,9 +199,9 @@ private:
 
 	virtual void DumpStats() const override;
 
-	virtual FPayloadActivityInfo GetAccumualtedPayloadActivityInfo() const override;
+	virtual FPayloadActivityInfo GetSystemStatistics() const override;
 
-	virtual void GetPayloadActivityInfo( GetPayloadActivityInfoFuncRef ) const override;
+	virtual TArray<FBackendStats> GetBackendStatistics() const override;
 
 	virtual void GatherAnalytics(TArray<FAnalyticsEventAttribute>& Attributes) const override;
 
diff --git a/Engine/Source/Editor/VirtualizationEditor/Private/SVirtualAssetsStatistics.cpp b/Engine/Source/Editor/VirtualizationEditor/Private/SVirtualAssetsStatistics.cpp
--- a/Engine/Source/Editor/VirtualizationEditor/Private/SVirtualAssetsStatistics.cpp
+++ b/Engine/Source/Editor/VirtualizationEditor/Private/SVirtualAssetsStatistics.cpp
@@ -336,7 +336,7 @@ TSharedRef<SWidget> SVirtualAssetsStatisticsDialog::GetGridPanel()
 	Panels.Push = CreateGridPanels(LOCTEXT("Upload", "Upload"));
 
 	int32 RowIndex = 2;
-	auto DisplayPayloadActivityInfo = [&StdMargin, &Color, &Font, &RowIndex, &Panels](const FString& DebugName, const FString& ConfigName, const FPayloadActivityInfo& PayloadActivityInfo)
+	auto DisplayPayloadActivityInfo = [&StdMargin, &Color, &Font, &RowIndex, &Panels](const FString& Name, const FPayloadActivityInfo& PayloadActivityInfo)
 		{
 			Panels.Names->AddSlot(0, RowIndex)
 			[
@@ -345,7 +345,7 @@ TSharedRef<SWidget> SVirtualAssetsStatisticsDialog::GetGridPanel()
 				.ColorAndOpacity(Color)
 				.Font(Font)
 				.Justification(ETextJustify::Left)
-				.Text(FText::FromString(ConfigName))
+				.Text(FText::FromString(Name))
 			];
 
 			auto FillPanelDetails = [&StdMargin, &Color, &Font, &RowIndex](TSharedPtr<SGridPanel>& Panel, const FPayloadActivityInfo::FActivity& Activity)
@@ -393,14 +393,18 @@ TSharedRef<SWidget> SVirtualAssetsStatisticsDialog::GetGridPanel()
 			RowIndex++;
 		};
 
-	IVirtualizationSystem::Get().GetPayloadActivityInfo(DisplayPayloadActivityInfo);
+	TArray<FBackendStats> BackendStats = IVirtualizationSystem::Get().GetBackendStatistics();
+	for (const FBackendStats& Stats : BackendStats)
+	{
+		DisplayPayloadActivityInfo(Stats.ConfigName, Stats.PayloadActivity);
+	}
 
-	FPayloadActivityInfo AccumulatedPayloadAcitvityInfo = IVirtualizationSystem::Get().GetAccumualtedPayloadActivityInfo();
+	FPayloadActivityInfo AccumulatedPayloadAcitvityInfo = IVirtualizationSystem::Get().GetSystemStatistics();
 
 	Color = TitleColor;
 	Font = TitleFont;
 
-	DisplayPayloadActivityInfo(FString("Total"), FString("Total"), AccumulatedPayloadAcitvityInfo);
+	DisplayPayloadActivityInfo(FString("Total"), AccumulatedPayloadAcitvityInfo);
 
 	return Panel;
 }
diff --git a/Engine/Source/Runtime/Core/Private/Virtualization/VirtualizationSystem.cpp b/Engine/Source/Runtime/Core/Private/Virtualization/VirtualizationSystem.cpp
--- a/Engine/Source/Runtime/Core/Private/Virtualization/VirtualizationSystem.cpp
+++ b/Engine/Source/Runtime/Core/Private/Virtualization/VirtualizationSystem.cpp
@@ -142,12 +142,12 @@ public:
 		// The null implementation will have no stats and nothing to log
 	}
 
-	virtual void GetPayloadActivityInfo(GetPayloadActivityInfoFuncRef) const override
+	virtual TArray<FBackendStats> GetBackendStatistics() const override
 	{
-		// The null implementation has no stats and nothing to invoke
+		return TArray<FBackendStats>();
 	}
 
-	virtual FPayloadActivityInfo GetAccumualtedPayloadActivityInfo() const override
+	virtual FPayloadActivityInfo GetSystemStatistics() const override
 	{
 		return FPayloadActivityInfo();
 	}
diff --git a/Engine/Source/Runtime/Core/Public/Virtualization/VirtualizationSystem.h b/Engine/Source/Runtime/Core/Public/Virtualization/VirtualizationSystem.h
--- a/Engine/Source/Runtime/Core/Public/Virtualization/VirtualizationSystem.h
+++ b/Engine/Source/Runtime/Core/Public/Virtualization/VirtualizationSystem.h
@@ -46,6 +46,21 @@ struct FPayloadActivityInfo
 	FActivity Cache;
 };
 
+/** Profiling data containing all activity relating to a single backend */
+struct FBackendStats
+{
+	/** Name of the backend as provided by the current VA graph found in ini:Engine */
+	FString ConfigName;
+	/** Full debug name of the backend */
+	FString DebugName;
+	
+	/** The storage type of the backend */
+	EStorageType Type;
+
+	/** Payload activity for the current backend */
+	FPayloadActivityInfo PayloadActivity;
+};
+
 /** Info about a rehydration operation */
 struct FRehydrationInfo
 {
@@ -823,11 +838,26 @@ public:
 
 	using GetPayloadActivityInfoFuncRef = TFunctionRef<void(const FString& DebugName, const FString& ConfigName, const FPayloadActivityInfo& PayloadInfo)>;
 
-	/** Access profiling info relating to payload activity per backend. Stats will only be collected if ENABLE_COOK_STATS is enabled.*/
-	virtual void GetPayloadActivityInfo( GetPayloadActivityInfoFuncRef ) const = 0;
+	UE_DEPRECATED(5.5, "Use IVirtualizationSystem::GetBackendStatistics instead")
+	virtual void GetPayloadActivityInfo(GetPayloadActivityInfoFuncRef Func) const
+	{
+		for (const FBackendStats& Stats : GetBackendStatistics())
+		{
+			Func(Stats.DebugName, Stats.ConfigName, Stats.PayloadActivity);
+		}
+	}
 
-	/** Access profiling info relating to accumulated payload activity. Stats will only be collected if ENABLE_COOK_STATS is enabled.*/
-	virtual FPayloadActivityInfo GetAccumualtedPayloadActivityInfo() const = 0;
+	UE_DEPRECATED(5.5, "Use IVirtualizationSystem::GetSystemStatistics instead")
+	virtual FPayloadActivityInfo GetAccumualtedPayloadActivityInfo() const
+	{
+		return GetSystemStatistics();
+	}
+	
+	/** Access profiling info relating to payload activity per backend. Stats will only be collected if ENABLE_COOK_STATS is enabled.*/
+	virtual TArray<FBackendStats> GetBackendStatistics() const = 0;
+	
+	/** Access profiling info for all backends combined. Stats will only be collected if ENABLE_COOK_STATS is enabled.*/
+	virtual FPayloadActivityInfo GetSystemStatistics() const = 0;
 
 	/** Gather analytics data. Stats will only be collected if ENABLE_COOK_STATS is enabled.*/
 	virtual void GatherAnalytics(TArray<FAnalyticsEventAttribute>& Attributes) const =0;
