Clean up VA initialization to avoid potential threading issue

#rnx

- TSAN picked up that FVirtualizationManager::bPendingBackendConnections could be accessed on multiple threads while being written too. Technically this is expected which is why the bool was checked a second time after acquiring the lock but to play things safe the bool has been changed to a std::atomic<bool>
- Changed InitCS from a FCriticalSystem (which can use an OS level primitive) to the less expensive FMutex since it's almost never going to be used.

diff --git a/Engine/Source/Developer/Virtualization/Private/VirtualizationManager.cpp b/Engine/Source/Developer/Virtualization/Private/VirtualizationManager.cpp
--- a/Engine/Source/Developer/Virtualization/Private/VirtualizationManager.cpp
+++ b/Engine/Source/Developer/Virtualization/Private/VirtualizationManager.cpp
@@ -1824,9 +1824,9 @@ void FVirtualizationManager::EnsureBackendConnections()
 	if (bPendingBackendConnections)
 	{
 		// Only allow one thread to initialize the system at a time
-		static FCriticalSection InitCS;
+		static FMutex Mutex;
 
-		FScopeLock _(&InitCS);
+		UE::TUniqueLock _(Mutex);
 		if (bPendingBackendConnections)
 		{
 			for (const TUniquePtr<IVirtualizationBackend>& Backend : AllBackends)
diff --git a/Engine/Source/Developer/Virtualization/Private/VirtualizationManager.h b/Engine/Source/Developer/Virtualization/Private/VirtualizationManager.h
--- a/Engine/Source/Developer/Virtualization/Private/VirtualizationManager.h
+++ b/Engine/Source/Developer/Virtualization/Private/VirtualizationManager.h
@@ -376,7 +376,7 @@ private:
 	FBackendArray PullEnabledBackends;
 
 	/** Do we have backends that have not yet tried connecting to their services */
-	bool bPendingBackendConnections;
+	std::atomic<bool> bPendingBackendConnections;
 
 	/** Our notification Event */
 	FOnNotification NotificationEvent;
