Add a way for us to find any EPayloadFilterReason that might be applied to a payload during the virtualization process.

#rnx

- Currently we apply filtering to payloads at two points, package save and during virtualization. The reasons applied during package save are stored int he payload trailer and are easy to access when making reports but we needed a way to find the second set of filters.
- This set up is fairly hacky as we are accessing the default FVirtualizationManager directly, this is because I wanted to avoid making a public interface change as I am not 100% sure that we should be trying to fix filtering this way and am leaning towards either moving all filtering to a single point (both have their advantages and disadvantages) or making the difference between these two stages more explicit and for either solution I don't want to be tied to any interface changes we'd need to make here.
- The new utilities header I added is internal so we can change it as we need, it is also marked as experimental so licensees should be avoiding it anyway.

diff --git a/Engine/Source/Developer/Virtualization/Internal/VirtualizationExperimentalUtilities.h b/Engine/Source/Developer/Virtualization/Internal/VirtualizationExperimentalUtilities.h
new file mode 100644
--- /dev/null
+++ b/Engine/Source/Developer/Virtualization/Internal/VirtualizationExperimentalUtilities.h
@@ -0,0 +1,31 @@
+// Copyright Epic Games, Inc. All Rights Reserved.
+
+#pragma once
+
+#include "Virtualization/VirtualizationSystem.h"
+
+/** 
+ * The functions in this header are not set in stone and are likely to be changed without warning.
+ * Once they are finalized they can be moved to the public utils header for general use.
+ */
+namespace UE::Virtualization::Utils
+{
+
+/**
+ * Adds filter reasons for a package that would only be calculated during the virtualization process and not stored in the payload trailer
+ * Currently this is used to fix up reporting in our commandlets and console commands so that it is easier for users to see why payloads
+ * have not been virtualized. This is only required because some filters are applied at package save time (when we know the owning asset
+ * type) and others at submit time.
+ * The two stage filtering does need fixing properly but I am not sure which way we want to do it, this is why we have an experimental
+ * utility to fix the reporting as I do not want to commit to an API change while still unsure on the direction we should take.
+ * 
+ * NOTE: Only works if the "Default" virtualization system is used.
+ * 
+ * @param PackagePath			Used to check for EPayloadFilterReason::Path
+ * @param SizeOnDisk			Used to check for EPayloadFilterReason::MinSize
+ * @param CurrentFilterFlags	The filters for the payload stored in the payload trailer.
+ * @return						The final filter reason used by the payload during the virtualization process.
+ */
+VIRTUALIZATION_API EPayloadFilterReason FixFilterFlags(FStringView PackagePath, uint64 SizeOnDisk, EPayloadFilterReason CurrentFilterFlags);
+
+} //namespace UE::Virtualization::Utils
\ No newline at end of file
diff --git a/Engine/Source/Developer/Virtualization/Private/VirtualizationManager.cpp b/Engine/Source/Developer/Virtualization/Private/VirtualizationManager.cpp
--- a/Engine/Source/Developer/Virtualization/Private/VirtualizationManager.cpp
+++ b/Engine/Source/Developer/Virtualization/Private/VirtualizationManager.cpp
@@ -2254,7 +2254,7 @@ bool FVirtualizationManager::ShouldVirtualizePackage(const FPackagePath& Package
 	return ShouldVirtualizeAsDefault();
 }
 
-bool FVirtualizationManager::ShouldVirtualize(const FString& Context) const
+bool FVirtualizationManager::ShouldVirtualize(FStringView Context) const
 {
 	// First see if we can convert the context from a raw string to a valid package path.
 	// If we can extract a package path then we should use the package filtering code
@@ -2360,6 +2360,26 @@ void FVirtualizationManager::GatherAnalytics(TArray<FAnalyticsEventAttribute>& A
 	}
 }
 
+EPayloadFilterReason FVirtualizationManager::FixFilterFlags(FStringView PackagePath, uint64 SizeOnDisk, EPayloadFilterReason CurrentFilterFlags)
+{
+	// We only apply new filters if the payload does not curently have any as stored filtered reasons would be checked first before
+	// we ever got to this part of the virtualization process.
+	if (CurrentFilterFlags == EPayloadFilterReason::None)
+	{
+		if (static_cast<int64>(SizeOnDisk) < MinPayloadLength)
+		{
+			return EPayloadFilterReason::MinSize;
+		}
+
+		if (!ShouldVirtualize(PackagePath))
+		{
+			return EPayloadFilterReason::Path;
+		}
+	}
+
+	return CurrentFilterFlags;
+}
+
 FString FVirtualizationManager::GetConnectionHelpUrl()
 {
 	return ConnectionHelpUrl;
diff --git a/Engine/Source/Developer/Virtualization/Private/VirtualizationManager.h b/Engine/Source/Developer/Virtualization/Private/VirtualizationManager.h
--- a/Engine/Source/Developer/Virtualization/Private/VirtualizationManager.h
+++ b/Engine/Source/Developer/Virtualization/Private/VirtualizationManager.h
@@ -209,6 +209,10 @@ private:
 	{
 		return NotificationEvent;
 	}
+
+public:
+	/** @see UE::Virtualization::Utils::FixFilterFlags */
+	EPayloadFilterReason FixFilterFlags(FStringView PackagePath, uint64 SizeOnDisk, EPayloadFilterReason CurrentFilterFlags);
 	
 private:
 
@@ -280,7 +284,7 @@ private:
 	 * 
 	 * @return True if the context should be virtualized and false if not.
 	 */
-	bool ShouldVirtualize(const FString& Context) const;
+	bool ShouldVirtualize(FStringView Context) const;
 
 	/** Determines if the default filtering behavior is to virtualize a payload or not */
 	bool ShouldVirtualizeAsDefault() const;
@@ -289,7 +293,7 @@ private:
 	bool ShouldRetryWhenUnattended() const;
 
 	void BroadcastEvent(TConstArrayView<FPullRequest> Ids, ENotification Event);
-	
+
 private:
 	// The following members are set from the config file
 
diff --git a/Engine/Source/Developer/Virtualization/Private/VirtualizationUtilities.cpp b/Engine/Source/Developer/Virtualization/Private/VirtualizationUtilities.cpp
--- a/Engine/Source/Developer/Virtualization/Private/VirtualizationUtilities.cpp
+++ b/Engine/Source/Developer/Virtualization/Private/VirtualizationUtilities.cpp
@@ -1,6 +1,7 @@
 // Copyright Epic Games, Inc. All Rights Reserved.
 
 #include "VirtualizationUtilities.h"
+#include "VirtualizationExperimentalUtilities.h"
 
 #include "HAL/PlatformProcess.h"
 #include "IO/IoHash.h"
@@ -9,6 +10,10 @@
 #include "Misc/StringBuilder.h"
 #include "UObject/PackageFileSummary.h"
 #include "UObject/PackageResourceManager.h"
+#include "Virtualization/VirtualizationSystem.h"
+#include "Virtualization/VirtualizationTypes.h"
+
+#include "VirtualizationManager.h"
 
 namespace UE::Virtualization::Utils
 {
@@ -167,4 +172,18 @@ bool IsProcessInteractive()
 	return true;
 }
 
+EPayloadFilterReason FixFilterFlags(FStringView PackagePath, uint64 SizeOnDisk, EPayloadFilterReason CurrentFilterFlags)
+{
+	if (IVirtualizationSystem::IsInitialized() && IVirtualizationSystem::GetSystemName() == FName("Default") && IVirtualizationSystem::Get().IsEnabled())
+	{
+		// Very hacky but should be safe if the system name is "Default". Allows us to do this without actually modifying the public API.
+		FVirtualizationManager& Manager = static_cast<FVirtualizationManager&>(IVirtualizationSystem::Get());
+
+		return Manager.FixFilterFlags(PackagePath, SizeOnDisk, CurrentFilterFlags);
+	}
+
+	return CurrentFilterFlags;
+}
+
 } // namespace UE::Virtualization::Utils
+
