The console command 'DumpPackagePayloadInfo' will now show the correct filter reasons in all cases.

#rnx

- In addition to fixing up the filter dispaly I have also added a FSlowTask as this command can take longer on machines with slower HD when run on many packages at once.

diff --git a/Engine/Source/Developer/Virtualization/Private/DumpPackagePayloadInfoCommand.cpp b/Engine/Source/Developer/Virtualization/Private/DumpPackagePayloadInfoCommand.cpp
--- a/Engine/Source/Developer/Virtualization/Private/DumpPackagePayloadInfoCommand.cpp
+++ b/Engine/Source/Developer/Virtualization/Private/DumpPackagePayloadInfoCommand.cpp
@@ -5,9 +5,13 @@
 #include "Logging/LogMacros.h"
 #include "Misc/PackageName.h"
 #include "Misc/PackagePath.h"
+#include "Misc/ScopedSlowTask.h"
 #include "UObject/PackageFileSummary.h"
 #include "UObject/PackageResourceManager.h"
 #include "UObject/PackageTrailer.h"
+#include "VirtualizationExperimentalUtilities.h"
+
+#define LOCTEXT_NAMESPACE "Virtualization"
 
 namespace UE
 {
@@ -95,6 +99,9 @@ void DumpPackagePayloadInfo(const TArray<FString>& Args)
 		return;
 	}
 
+	TRACE_CPUPROFILER_EVENT_SCOPE(DumpPackagePayloadInfo);
+	FScopedSlowTask Progress(0.0f, LOCTEXT("VAFindPayloadInfo", "Finding payload info..."));
+
 	for (const FString& Arg : Args)
 	{
 		FString PathString;
@@ -149,6 +156,9 @@ void DumpPackagePayloadInfo(const TArray<FString>& Args)
 			for (int32 Index = 0; Index < LocalPayloadIds.Num(); ++Index)
 			{
 				FPayloadInfo Info = Trailer.GetPayloadInfo(LocalPayloadIds[Index]);
+
+				Info.FilterFlags = UE::Virtualization::Utils::FixFilterFlags(PathString, Info.CompressedSize, Info.FilterFlags);
+
 				UE_LOG(LogVirtualization, Display, TEXT("%02d    | %s | %-10s | %s"),
 					Index,
 					*LexToString(LocalPayloadIds[Index]),
@@ -182,3 +192,5 @@ static FAutoConsoleCommand CCmdDumpPayloadToc = FAutoConsoleCommand(
 #endif //WITH_EDITORONLY_DATA
 
 } // namespace UE
+
+#undef LOCTEXT_NAMESPACE
