Downgrade some VA logging from 'error' to 'warning'

#rnx

- The system now has multiple ways for the user or process to attempt to retry failed payload pull operations. If the retry is successful but we logged errors for the initial failure than some CIS systems will consider that job a failure which can interrupt builds.
- By downgrading the error to a warning we avoid this. The system will still log a final error if the retries are aborted or otherwise unsuccessful.
- I did consider only demote the logging for unattended builds but most users will not be able to do anything useful with the logged messages anyway and will have to reach out for technical support so there is no great gain in keeping them as errors and I'd prefer to try and keep the code paths as similar as possible unless we have good reason.
- Added code comments to make intent more clear.

diff --git a/Engine/Source/Developer/Virtualization/Private/VirtualizationManager.cpp b/Engine/Source/Developer/Virtualization/Private/VirtualizationManager.cpp
--- a/Engine/Source/Developer/Virtualization/Private/VirtualizationManager.cpp
+++ b/Engine/Source/Developer/Virtualization/Private/VirtualizationManager.cpp
@@ -2065,13 +2065,21 @@ FVirtualizationManager::ErrorHandlingResult FVirtualizationManager::OnPayloadPul
 		return ErrorHandlingResult::AcceptFailedPayloads;
 	}
 
+	// Note that when handling failed pull errors there are several code paths that allow the process or the user to retry to the failed
+	// pull operation in case it was caused by a temporary infrastructure issue. So until we are certain that no more retrying will be
+	// attempted we must log everything as a 'warning' rather than an 'error' as many CIS systems will treat even single logged error as
+	// a reason to consider a job failed.
+
 	for (const FIoHash& FailedPayload : Requests.GetFailedPayloads())
 	{
-		UE_LOG(LogVirtualization, Error, TEXT("Failed to pull payload '%s'"), *LexToString(FailedPayload));
+		UE_LOG(LogVirtualization, Warning, TEXT("Failed to pull payload '%s'"), *LexToString(FailedPayload));
 	}
 
 	static FCriticalSection CriticalSection;
 
+	// Only allow a single thread to make an active error report. The most likely cause of an error is an infrastructure outage which would
+	// cause all threads making requests to fail at the same time as we don't want to spam the user with dialogs. If they (or the unattended
+	// process) requests a retry then we should apply that to all threads currently encountering issues at the same time.
 	if (CriticalSection.TryLock())
 	{
 		const FText Title(LOCTEXT("VAPullTitle", "Failed to pull virtualized data!"));
@@ -2117,7 +2125,7 @@ FVirtualizationManager::ErrorHandlingResult FVirtualizationManager::OnPayloadPul
 			{
 				// Only log as a warning while we are retrying so that if we do recover out logging will
 				// not cause the running process to be considered a failure due to an error.
-				UE_LOG(LogVirtualization, Warning, TEXT("Failed to pull payload(s) on attempt (%d/%d)"), UnattendedFailureMsgCount.load(), UnattendedRetryCount);
+				UE_LOG(LogVirtualization, Warning, TEXT("Failed to pull from any backend on attempt (%d/%d)"), UnattendedFailureMsgCount.load(), UnattendedRetryCount);
 
 				Result = EAppReturnType::Yes; // Attempt to retry the failed pull
 
