Restore log warning when the virtualization source control backend retries a pull operation.

#rnx

- The log message was previously commented out during some refactoring but with the current push to only error on absolute failure it is important to record when multiple pull attempts are being made.

diff --git a/Engine/Source/Developer/Virtualization/Private/VirtualizationSourceControlBackend.cpp b/Engine/Source/Developer/Virtualization/Private/VirtualizationSourceControlBackend.cpp
--- a/Engine/Source/Developer/Virtualization/Private/VirtualizationSourceControlBackend.cpp
+++ b/Engine/Source/Developer/Virtualization/Private/VirtualizationSourceControlBackend.cpp
@@ -556,7 +556,7 @@ bool FSourceControlBackend::PullData(TArrayView<FPullRequest> Requests, EPullFla
 		// Only warn if the backend is configured to retry
 		if (Retries != 0)
 		{
-		//	UE_LOG(LogVirtualization, Warning, TEXT("[%s] Failed to download '%s' retrying (%d/%d) in %dms..."), *GetDebugName(), DepotPath.ToString(), Retries, RetryCount, RetryWaitTimeMS);
+			UE_LOG(LogVirtualization, Warning, TEXT("[%s] Failed to pull payloads! Attempt (%d/%d) retrying in %d(ms)..."), *GetDebugName(), Retries, RetryCount, RetryWaitTimeMS);
 			FPlatformProcess::SleepNoStats(static_cast<float>(RetryWaitTimeMS) * 0.001f);
 		}
 
