Improve error reporting if FSourceControlBackend fails to pull a payload.

#rnx

- We now only report errors encountered on the final retry rather than once per iteration. This does run the risk of discarding useful info but in practice this just resulted in the same error(s) being logged over and over again until the retry counter ran out and made the logs spammy.
- This also now results in us logging that the final retry attempt was tried, before we would always exit on retry-1 which made it uncertain if the final attempt ran or not.
- We now log the perforce errors before our connection errors as they are more important.
- Add an explicit error if a none positive number of pull attempts were requested. The retry count is poorly named and has no validation, this will be addressed in a future work item.

diff --git a/Engine/Source/Developer/Virtualization/Private/VirtualizationSourceControlBackend.cpp b/Engine/Source/Developer/Virtualization/Private/VirtualizationSourceControlBackend.cpp
--- a/Engine/Source/Developer/Virtualization/Private/VirtualizationSourceControlBackend.cpp
+++ b/Engine/Source/Developer/Virtualization/Private/VirtualizationSourceControlBackend.cpp
@@ -551,7 +551,15 @@ bool FSourceControlBackend::PullData(TArrayView<FPullRequest> Requests, EPullFla
 
 	bool bConnectionFailed = false;
 
-	while (Retries < RetryCount)
+	// Note that 'RetryCount' is badly named, it's actually the number of attempts we can make, not 1 attempt + retries
+
+	if (RetryCount <= 0)
+	{
+		UE_LOG(LogVirtualization, Warning, TEXT("[%s] Retry count is set to 0, no attempts can be made!"), *GetDebugName());
+		return false;
+	}
+
+	while (true)
 	{
 		// Only warn if the backend is configured to retry
 		if (Retries != 0)
@@ -560,17 +568,17 @@ bool FSourceControlBackend::PullData(TArrayView<FPullRequest> Requests, EPullFla
 			FPlatformProcess::SleepNoStats(static_cast<float>(RetryWaitTimeMS) * 0.001f);
 		}
 
+		TSharedRef<FDownloadFile> DownloadCommand = ISourceControlOperation::Create<FDownloadFile>(FDownloadFile::EVerbosity::None);
+
 #if IS_SOURCE_CONTROL_THREAD_SAFE
-		TSharedRef<FDownloadFile, ESPMode::ThreadSafe> DownloadCommand = ISourceControlOperation::Create<FDownloadFile>(FDownloadFile::EVerbosity::None);
 		const bool bOperationSuccess = SCCProvider->Execute(DownloadCommand, DepotPaths, EConcurrency::Synchronous) == ECommandResult::Succeeded;
 #else
-		TSharedRef<FDownloadFile> DownloadCommand = ISourceControlOperation::Create<FDownloadFile>(FDownloadFile::EVerbosity::None);
 		const bool bOperationSuccess = SCCProvider->TryToDownloadFileFromBackgroundThread(DownloadCommand, DepotPaths);
 #endif
 
 		// Check and assign what payloads we did find, even if the download command failed we might have been able to download some
 		// files and since we found them we might as well return them.
-		for(int32 Index = 0; Index < Requests.Num(); ++Index)
+		for (int32 Index = 0; Index < Requests.Num(); ++Index)
 		{
 			// The payload was created by FCompressedBuffer::Compress so we can return it as a FCompressedBuffer.
 			FSharedBuffer Buffer = DownloadCommand->GetFileData(DepotPaths[Index]);
@@ -592,21 +600,35 @@ bool FSourceControlBackend::PullData(TArrayView<FPullRequest> Requests, EPullFla
 			return false;
 		}
 
-		bConnectionFailed = DownloadCommand->GetResultInfo().DidConnectionFail();
+		if (++Retries >= RetryCount)
+		{
+			UE_LOG(LogVirtualization, Warning, TEXT("[%s] Failed to pull payloads! Attempt (%d/%d)"), *GetDebugName(), Retries, RetryCount);
+			
+			FTextBuilder ErrorBuilder;
+			if (DownloadCommand->GetResultInfo().DidConnectionFail())
+			{
+				TMap<ISourceControlProvider::EStatus, FString> SCPStatus = SCCProvider->GetStatus();
 
-		Retries++;
-	}
+				FText ErrorMsg = FText::Format(LOCTEXT("VA_SCP", "Failed to connect to perforce server '{0}' with username '{1}'"),
+					FText::FromString(SCPStatus[ISourceControlProvider::EStatus::Port]),
+					FText::FromString(SCPStatus[ISourceControlProvider::EStatus::User]));
 
-	if (bConnectionFailed)
-	{
-		TMap<ISourceControlProvider::EStatus, FString> SCPStatus = SCCProvider->GetStatus();
-		
-		OutErrors = FText::Format(LOCTEXT("VA_SCP", "Failed to connect to perforce server '{0}' with username '{1}'"),
-			FText::FromString(SCPStatus[ISourceControlProvider::EStatus::Port]), 
-			FText::FromString(SCPStatus[ISourceControlProvider::EStatus::User]));
-	}
+				ErrorBuilder.AppendLine(ErrorMsg);
+			}
 
-	return false;
+			for (const FText& ErrorMsg : DownloadCommand->GetResultInfo().ErrorMessages)
+			{
+				ErrorBuilder.AppendLine(ErrorMsg);
+			}
+
+			if (!ErrorBuilder.IsEmpty())
+			{
+				OutErrors = ErrorBuilder.ToText();
+			}
+
+			return false;
+		}
+	}
 }
 
 bool FSourceControlBackend::DoesPayloadExist(const FIoHash& Id)
