Do not log errors when failing to connect to the VA source control backend if the process is non-interactive.

#rnx

- When running in a non interaction mode we only want to log errors when VA completely fails and cannot continue any further. If there is a chance we can retry a failed operaiton we should only log a warning. This is to prevent things like CIS from failing because of a logged error that we later recovered on.
- At the moment a failed connection cannot really be recovered from but there is still a chance that a different backend might succeed in a future pull operation so we shouldn't log an error.

diff --git a/Engine/Source/Developer/Virtualization/Private/VirtualizationSourceControlBackend.cpp b/Engine/Source/Developer/Virtualization/Private/VirtualizationSourceControlBackend.cpp
--- a/Engine/Source/Developer/Virtualization/Private/VirtualizationSourceControlBackend.cpp
+++ b/Engine/Source/Developer/Virtualization/Private/VirtualizationSourceControlBackend.cpp
@@ -1297,10 +1297,17 @@ void FSourceControlBackend::OnConnectionError(FText ErrorMessage)
 		TokenizedMsg->AddToken(FURLToken::Create(ConnectionHelpUrl, LOCTEXT("URL", "Additional connection help")));
 	}
 
+	if (!Utils::IsProcessInteractive())
+	{
+		UE_LOG(LogVirtualization, Warning, TEXT("%s"), *ErrorMessage.ToString());
+		UE_LOG(LogVirtualization, Warning, TEXT("%s"), *UserMessage.ToString());
+
+		return;
+	}
+	
 	if (::IsInGameThread())
 	{
 		// If we are on the game thread we can post the error message immediately
-
 		{
 			FMessageLog Log("LogVirtualization");
 			Log.Error(ErrorMessage);
