Improve VA backend connection support when running with lazy initialization options.

- We no longer call FVirtualizationManager::EnsureBackendConnections before running an operation, as this would iterate over all available backends at once and make sure each one had made an attempt to connect before allowing the operation to ruin.
-- As this would only run once this meant we could not recover from a backend disconnecting.
-- It also meant that we could try to connect to backends that we never use.
- We now call IVirtualizationBackend::EnsureConnection on each backend each time it is to be used, so if the first backend in the list is always hit we will never try to connect to the subsequent backends. This also means that backends can attempt to reconnect if they failed previous connections.
- The previous legacy behavior can be restored by setting "-ini:Engine:[Core.VirtualizationModule]:AllowBackendReconnection=false"

### Misc
- Removed FVirtualizationManager::EnsureBackendConnections entirely and it's associated members.
- Re-order IVirtualizationBackend members and set the size of the enum EConnectionStatus to reduce the size of IVirtualizationBackend from 56 bytes to 48 bytes.
- When profiling, the newly added mutex lock was pretty much a rounding error compared to the rest of the work and isn't something we need to worry about.

diff --git a/Engine/Config/BaseEngine.ini b/Engine/Config/BaseEngine.ini
--- a/Engine/Config/BaseEngine.ini
+++ b/Engine/Config/BaseEngine.ini
@@ -1558,6 +1558,7 @@ LazyInitConnections=False
 DisableLazyInitIfInteractive=False
 UseLegacyErrorHandling=False
 ForceCachingOnPull=False
+AllowBackendReconnection=True
 PullErrorAdditionalMsg=""
 ConnectionHelpUrl=""
 UnattendedRetryCount=0
diff --git a/Engine/Source/Developer/Virtualization/Private/VirtualizationManager.cpp b/Engine/Source/Developer/Virtualization/Private/VirtualizationManager.cpp
--- a/Engine/Source/Developer/Virtualization/Private/VirtualizationManager.cpp
+++ b/Engine/Source/Developer/Virtualization/Private/VirtualizationManager.cpp
@@ -24,7 +24,6 @@
 
 #define LOCTEXT_NAMESPACE "Virtualization"
 
-
 namespace UE::Virtualization
 {
 UE_REGISTER_VIRTUALIZATION_SYSTEM(UE::Virtualization::FVirtualizationManager, Default);
@@ -543,7 +542,7 @@ FVirtualizationManager::FVirtualizationManager()
 	, bLazyInitConnections(false)
 	, bUseLegacyErrorHandling(true)
 	, bForceCachingOnPull(false)
-	, bPendingBackendConnections(false)
+	, bAllowBackendReconnection(true)
 {
 }
 
@@ -748,8 +747,6 @@ bool FVirtualizationManager::PushData(TArrayView<FPushRequest> Requests, EStorag
 		return true;
 	}
 
-	EnsureBackendConnections();
-
 	FConditionalScopeLock _(&DebugValues.ForceSingleThreadedCS, DebugValues.bSingleThreaded);
 
 	// TODO: Note that all push operations are currently synchronous, probably 
@@ -761,9 +758,9 @@ bool FVirtualizationManager::PushData(TArrayView<FPushRequest> Requests, EStorag
 	
 	for (IVirtualizationBackend* Backend : Backends)
 	{
-		if (Backend->GetConnectionStatus() != IVirtualizationBackend::EConnectionStatus::Connected)
+		if (Backend->TryConnect(bAllowBackendReconnection) != IVirtualizationBackend::EConnectionStatus::Connected)
 		{
-			UE_LOG(LogVirtualization, Verbose, TEXT("Cannot push to backend '%s' as it is not connected"), *Backend->GetDebugName());
+			UE_LOG(LogVirtualization, Verbose, TEXT("Cannot push to backend  '%s' as it was unable to connect"), *Backend->GetDebugName());
 			continue;
 		}
 
@@ -822,8 +819,6 @@ bool FVirtualizationManager::PullData(TArrayView<FPullRequest> Requests)
 		return false;
 	}
 
-	EnsureBackendConnections();
-
 	FConditionalScopeLock _(&DebugValues.ForceSingleThreadedCS, DebugValues.bSingleThreaded);
 
 	BroadcastEvent(Requests, IVirtualizationSystem::PullBegunNotification);
@@ -887,16 +882,14 @@ EQueryResult FVirtualizationManager::QueryPayloadStatuses(TArrayView<const FIoHa
 	HitCount.SetNum(Ids.Num());
 	Results.SetNum(Ids.Num());
 
-	EnsureBackendConnections();
-
 	{
 		FConditionalScopeLock _(&DebugValues.ForceSingleThreadedCS, DebugValues.bSingleThreaded);
 
 		for (IVirtualizationBackend* Backend : Backends)
 		{
-			if (Backend->GetConnectionStatus() != IVirtualizationBackend::EConnectionStatus::Connected)
+			if (Backend->TryConnect(bAllowBackendReconnection) != IVirtualizationBackend::EConnectionStatus::Connected)
 			{
-				UE_LOG(LogVirtualization, Verbose, TEXT("Cannot query backend '%s' as it is not connected"), *Backend->GetDebugName());
+				UE_LOG(LogVirtualization, Verbose, TEXT("Cannot query backend  '%s' as it was unable to connect"), *Backend->GetDebugName());
 				continue;
 			}
 
@@ -1226,6 +1219,11 @@ void FVirtualizationManager::ApplySettingsFromConfigFiles(const FConfigFile& Con
 		UE_LOG(LogVirtualization, Display, TEXT("\tForceCachingOnPull : %s"), bForceCachingOnPull ? TEXT("true") : TEXT("false"));
 	}
 
+	{
+		ConfigFile.GetBool(ConfigSection, TEXT("AllowBackendReconnection"), bAllowBackendReconnection);
+		UE_LOG(LogVirtualization, Display, TEXT("\tAllowBackendReconnection : %s"), bAllowBackendReconnection ? TEXT("true") : TEXT("false"));
+	}
+
 	{
 		ConfigFile.GetInt(ConfigSection, TEXT("UnattendedRetryTimer"), UnattendedRetryTimer);
 		ConfigFile.GetInt(ConfigSection, TEXT("UnattendedRetryCount"), UnattendedRetryCount);
@@ -1804,11 +1802,7 @@ void FVirtualizationManager::AddBackend(TUniquePtr<IVirtualizationBackend> Backe
 	// In the future this will be made async to avoid blocking the GameThread on startup
 	if (!bLazyInitConnections)
 	{
-		BackendRef->Connect();
-	}
-	else
-	{
-		bPendingBackendConnections = true;
+		BackendRef->TryConnect(bAllowBackendReconnection);
 	}
 
 	COOK_STAT(Profiling::CreateStats(*BackendRef));
@@ -1819,31 +1813,12 @@ bool FVirtualizationManager::IsPersistentBackend(IVirtualizationBackend& Backend
 	return PersistentStorageBackends.Contains(&Backend);
 }
 
-void FVirtualizationManager::EnsureBackendConnections()
-{
-	if (bPendingBackendConnections)
-	{
-		// Only allow one thread to initialize the system at a time
-		static FMutex Mutex;
-
-		UE::TUniqueLock _(Mutex);
-		if (bPendingBackendConnections)
-		{
-			for (const TUniquePtr<IVirtualizationBackend>& Backend : AllBackends)
-			{
-				Backend->Connect();
-			}
-
-			bPendingBackendConnections = false;
-		}
-	}
-}
-
 void FVirtualizationManager::CachePayloads(TArrayView<FPushRequest> Requests, const IVirtualizationBackend* BackendSource, IVirtualizationBackend::EPushFlags Flags)
 {
 	TRACE_CPUPROFILER_EVENT_SCOPE(FVirtualizationManager::CachePayload);
 
-	check(!bPendingBackendConnections);
+	//check(!bPendingBackendConnections);
+	// TODO: Connection?
 
 	// We start caching at the first (assumed to be fastest) local cache backend. 
 	for (IVirtualizationBackend* BackendToCache : CacheStorageBackends)
@@ -1982,9 +1957,9 @@ void FVirtualizationManager::PullDataFromAllBackends(TArrayView<FPullRequest> Re
 				continue;
 			}
 
-			if (Backend->GetConnectionStatus() != IVirtualizationBackend::EConnectionStatus::Connected)
+			if (Backend->TryConnect(bAllowBackendReconnection) != IVirtualizationBackend::EConnectionStatus::Connected)
 			{
-				UE_LOG(LogVirtualization, Verbose, TEXT("Cannot pull from backend '%s' as it is not connected"), *Backend->GetDebugName());
+				UE_LOG(LogVirtualization, Verbose, TEXT("Cannot pull from backend '%s' as it was unable to connect"), *Backend->GetDebugName());
 				continue;
 			}
 
diff --git a/Engine/Source/Developer/Virtualization/Private/VirtualizationManager.h b/Engine/Source/Developer/Virtualization/Private/VirtualizationManager.h
--- a/Engine/Source/Developer/Virtualization/Private/VirtualizationManager.h
+++ b/Engine/Source/Developer/Virtualization/Private/VirtualizationManager.h
@@ -135,6 +135,9 @@ struct FAnalyticsEventAttribute;
  *												before it was pulled from a backend later in the hierarchy. Can be used to try and skip
  *												expensive existence checks, or if a backend is in a bad state where it believes it has the payload
  *												but is unable to actually return the data. [Default=false]
+ * AllowBackendReconnection [bool]:				When true backends that have previously failed to connect to their services can retry to connect
+ *												each time they are accessed. When false the legacy behavior is invoked and backends may only attempt
+ *												to connect a single time and after that will be unusable. [Default=true]
  * UnattendedRetryCount [int32]:				How many times the process should retry pulling payloads after a failure is encountered if the
 												process is unattended. Usually when a payload pull fails we ask the user to try and fix the issue
 												and retry, but in unattended mode we just log an error and terminate the process. In some cases
@@ -240,8 +243,6 @@ private:
 
 	bool IsPersistentBackend(IVirtualizationBackend& Backend);
 
-	void EnsureBackendConnections();
-
 	void CachePayloads(TArrayView<FPushRequest> Requests, const IVirtualizationBackend* BackendSource, IVirtualizationBackend::EPushFlags Flags);
 
 	bool TryCacheDataToBackend(IVirtualizationBackend& Backend, TArrayView<FPushRequest> Requests, IVirtualizationBackend::EPushFlags Flags);
@@ -344,6 +345,9 @@ private:
 	/** When true IVirtualizationBackend::EPushFlags::Force will be passed to backends that need to cache payloads when pulling */
 	bool bForceCachingOnPull;
 
+	/** When true backends will be able to try to reconnect to their service's if previous connection attempts failed */
+	bool bAllowBackendReconnection;
+
 	/** An additional error message to display when pulling payloads fails */
 	FString PullErrorAdditionalMsg;
 
@@ -379,9 +383,6 @@ private:
 	 */
 	FBackendArray PullEnabledBackends;
 
-	/** Do we have backends that have not yet tried connecting to their services */
-	std::atomic<bool> bPendingBackendConnections;
-
 	/** Our notification Event */
 	FOnNotification NotificationEvent;
 
diff --git a/Engine/Source/Developer/Virtualization/Public/IVirtualizationBackend.h b/Engine/Source/Developer/Virtualization/Public/IVirtualizationBackend.h
--- a/Engine/Source/Developer/Virtualization/Public/IVirtualizationBackend.h
+++ b/Engine/Source/Developer/Virtualization/Public/IVirtualizationBackend.h
@@ -2,6 +2,8 @@
 
 #pragma once
 
+#include "Async/Mutex.h"
+#include "Async/UniqueLock.h"
 #include "Compression/CompressedBuffer.h"
 #include "Containers/StringView.h"
 #include "Features/IModularFeature.h"
@@ -40,7 +42,7 @@ public:
 	FRIEND_ENUM_CLASS_FLAGS(EOperations);
 
 	/** Status of the backends connection to its services */
-	enum class EConnectionStatus
+	enum class EConnectionStatus : uint8
 	{
 		/** No connection attempt has been made */
 		None,
@@ -71,10 +73,10 @@ public:
 protected:
 
 	IVirtualizationBackend(FStringView InConfigName, FStringView InDebugName, EOperations InSupportedOperations)
-		: SupportedOperations(InSupportedOperations)
-		, DebugDisabledOperations(EOperations::None)
-		, ConfigName(InConfigName)
+		: ConfigName(InConfigName)
 		, DebugName(InDebugName)
+		, SupportedOperations(InSupportedOperations)
+		, DebugDisabledOperations(EOperations::None)
 	{
 		checkf(InSupportedOperations != EOperations::None, TEXT("Cannot create a backend with no supported operations!"));
 	}
@@ -96,15 +98,35 @@ public:
 	 */
 	virtual bool Initialize(const FString& ConfigEntry) = 0;
 
-	/**
+	
+	UE_DEPRECATED(5.6, "Call IVirtualizationBackend::TryConnect instead")
+	virtual void Connect() final
+	{
+		TryConnect(false);
+	}
+
+	/** 
 	 * Attempt to connect the backend to its services if not already connected.
+	 * 
+	 * @praram bAllowReconnect	When true the backend will attempt to connect to it's services even
+	 *							if it has previously failed to do so. When false the backend will
+	 *							only ever attempt to connect a single time with that result being
+	 *							final.
+	 * 
+	 * @return The current connection status of the backend. @see EConnectionStatus
 	 */
-	virtual void Connect()
+	EConnectionStatus TryConnect(bool bAllowReconnect)
 	{
-		if (ConnectionStatus != EConnectionStatus::Connected)
+		TRACE_CPUPROFILER_EVENT_SCOPE(IVirtualizationBackend::TryConnect);
+
+		UE::TUniqueLock _(ConnectionMutex);
+
+		if ((bAllowReconnect && ConnectionStatus != EConnectionStatus::Connected) || (!bAllowReconnect && ConnectionStatus == EConnectionStatus::None))
 		{
 			ConnectionStatus = OnConnect();
 		}
+		
+		return ConnectionStatus;
 	}
 
 	/**
@@ -114,6 +136,7 @@ public:
 	 */
 	EConnectionStatus GetConnectionStatus() const
 	{
+		UE::TUniqueLock _(ConnectionMutex);
 		return ConnectionStatus;
 	}
 
@@ -266,19 +289,23 @@ private:
 
 private:
 
-	/** The operations that this backend supports */
-	EOperations SupportedOperations;
-
-	EOperations DebugDisabledOperations;
-
 	/** The name assigned to the backend by the virtualization graph */
 	FString ConfigName;
 
 	/** Combination of the backend type and the name used to create it in the virtualization graph */
 	FString DebugName;
 
+	/** Protects threaded access to 'ConnectionStatus' */
+	mutable UE::FMutex ConnectionMutex;
+
 	/** The status of the connection to the backends service */
 	EConnectionStatus ConnectionStatus = EConnectionStatus::None;
+
+	/** The operations that this backend supports */
+	EOperations SupportedOperations;
+
+	/** The operations that have been disabled for debugging purposes */
+	EOperations DebugDisabledOperations;
 };
 
 /** 
