Move ::TryFindProject from UVT to the virtualization utilities so that other systems can re-use the same logic.

#rnx

- The only real change is that the function now takes the project extension to look for. All of the logic remains the same.
- Added some additional code documentation on changes to the function that we might want to do in the future.

diff --git a/Engine/Source/Developer/Virtualization/Private/VirtualizationUtilities.cpp b/Engine/Source/Developer/Virtualization/Private/VirtualizationUtilities.cpp
--- a/Engine/Source/Developer/Virtualization/Private/VirtualizationUtilities.cpp
+++ b/Engine/Source/Developer/Virtualization/Private/VirtualizationUtilities.cpp
@@ -3,6 +3,7 @@
 #include "VirtualizationUtilities.h"
 #include "VirtualizationExperimentalUtilities.h"
 
+#include "HAL/FileManager.h"
 #include "HAL/PlatformProcess.h"
 #include "IO/IoHash.h"
 #include "Misc/App.h"
@@ -12,7 +13,6 @@
 #include "UObject/PackageResourceManager.h"
 #include "Virtualization/VirtualizationSystem.h"
 #include "Virtualization/VirtualizationTypes.h"
-
 #include "VirtualizationManager.h"
 
 namespace UE::Virtualization::Utils
@@ -185,5 +185,79 @@ EPayloadFilterReason FixFilterFlags(FStringView PackagePath, uint64 SizeOnDisk,
 	return CurrentFilterFlags;
 }
 
+bool TryFindProject(const FString& PackagePath, const FString& Extension, FString& ProjectFilePath, FString& PluginFilePath)
+{
+	TRACE_CPUPROFILER_EVENT_SCOPE(UE::Virtualization::Utils::TryFindProject);
+
+	// TODO: This could be heavily optimized by caching known project files maybe with the use of FDirectoryTree
+	// TODO: Relying on the known content & plugin directory conventions helps optimize this code but is fragile.
+	// But if we started checking every directory then we'd need to start caching the results.
+	int32 ContentIndex = PackagePath.Find(TEXT("/content/"), ESearchCase::IgnoreCase, ESearchDir::FromEnd);
+
+	// Early out if there is not a single content directory in the path
+	if (ContentIndex == INDEX_NONE)
+	{
+		UE_LOG(LogVirtualization, Warning, TEXT("'%s' is not under a content directory"), *PackagePath);
+		return false;
+	}
+
+	while (ContentIndex != INDEX_NONE)
+	{
+		// Assume that the project directory is the parent of the /content/ directory
+		FString ProjectDirectory = PackagePath.Left(ContentIndex);
+		FString PluginDirectory;
+
+		TArray<FString> ProjectFile;
+		TArray<FString> PluginFile;
+
+		IFileManager::Get().FindFiles(ProjectFile, *ProjectDirectory, *Extension);
+
+		if (ProjectFile.IsEmpty())
+		{
+			// If there was no project file, the package could be in a plugin, so lets check for that
+			PluginDirectory = ProjectDirectory;
+			IFileManager::Get().FindFiles(PluginFile, *PluginDirectory, TEXT(".uplugin"));
+
+			if (PluginFile.Num() == 1)
+			{
+				PluginFilePath = PluginDirectory / PluginFile[0];
+
+				// We have a valid plugin file, so we should be able to find a /plugins/ directory which will be just below the project directory
+				const int32 PluginIndex = PluginDirectory.Find(TEXT("/plugins/"), ESearchCase::IgnoreCase, ESearchDir::FromEnd);
+				if (PluginIndex != INDEX_NONE)
+				{
+					// We found the plugin root directory so the one above it should be the project directory
+					ProjectDirectory = PluginDirectory.Left(PluginIndex);
+					IFileManager::Get().FindFiles(ProjectFile, *ProjectDirectory, *Extension);
+				}
+			}
+			else if (PluginFile.Num() > 1)
+			{
+				UE_LOG(LogVirtualization, Warning, TEXT("Found multiple .uplugin files for '%s' at '%s'"), *PackagePath, *PluginDirectory);
+				return false;
+			}
+		}
+
+		if (ProjectFile.Num() == 1)
+		{
+			ProjectFilePath = ProjectDirectory / ProjectFile[0];
+			return true;
+		}
+		else if (!ProjectFile.IsEmpty())
+		{
+			UE_LOG(LogVirtualization, Warning, TEXT("Found multiple .uproject files for '%s' at '%s'"), *PackagePath, *ProjectDirectory);
+			return false;
+		}
+
+		// Could be more than one content directory in the path so lets keep looking
+		ContentIndex = PackagePath.Find(TEXT("/content/"), ESearchCase::IgnoreCase, ESearchDir::FromEnd, ContentIndex);
+	}
+
+	// We found one or more content directories but none of them contained a project file
+	UE_LOG(LogVirtualization, Warning, TEXT("Failed to find project file for '%s'"), *PackagePath);
+
+	return false;
+}
+
 } // namespace UE::Virtualization::Utils
 
diff --git a/Engine/Source/Developer/Virtualization/Public/VirtualizationUtilities.h b/Engine/Source/Developer/Virtualization/Public/VirtualizationUtilities.h
--- a/Engine/Source/Developer/Virtualization/Public/VirtualizationUtilities.h
+++ b/Engine/Source/Developer/Virtualization/Public/VirtualizationUtilities.h
@@ -83,4 +83,6 @@ bool ExpandEnvironmentVariables(FStringView InputPath, FStringBuilderBase& OutEx
  * */
 bool IsProcessInteractive();
 
+bool TryFindProject(const FString& PackagePath, const FString& Extension, FString& ProjectFilePath, FString& PluginFilePath);
+
 } // namespace UE::Virtualization::Utils
diff --git a/Engine/Source/Programs/UnrealVirtualizationTool/Private/UnrealVirtualizationToolApp.cpp b/Engine/Source/Programs/UnrealVirtualizationTool/Private/UnrealVirtualizationToolApp.cpp
--- a/Engine/Source/Programs/UnrealVirtualizationTool/Private/UnrealVirtualizationToolApp.cpp
+++ b/Engine/Source/Programs/UnrealVirtualizationTool/Private/UnrealVirtualizationToolApp.cpp
@@ -434,7 +434,7 @@ bool FUnrealVirtualizationToolApp::TrySortFilesByProject(const TArray<FString>&
 		FString ProjectFilePath;
 		FString PluginFilePath;
 
-		if (TryFindProject(PackagePath, ProjectFilePath, PluginFilePath))
+		if (Utils::TryFindProject(PackagePath, TEXT(".uproject"), ProjectFilePath, PluginFilePath))
 		{
 			FProject& Project = FindOrAddProject(MoveTemp(ProjectFilePath));
 			if (PluginFilePath.IsEmpty())
@@ -453,78 +453,6 @@ bool FUnrealVirtualizationToolApp::TrySortFilesByProject(const TArray<FString>&
 	return true;
 }
 
-bool FUnrealVirtualizationToolApp::TryFindProject(const FString& PackagePath, FString& ProjectFilePath, FString& PluginFilePath) const
-{
-	TRACE_CPUPROFILER_EVENT_SCOPE(TryFindProject);
-
-	// TODO: This could be heavily optimized by caching known project files
-
-	int32 ContentIndex = PackagePath.Find(TEXT("/content/"), ESearchCase::IgnoreCase, ESearchDir::FromEnd);
-
-	// Early out if there is not a single content directory in the path
-	if (ContentIndex == INDEX_NONE)
-	{
-		UE_LOG(LogVirtualizationTool, Warning, TEXT("'%s' is not under a content directory"), *PackagePath);
-		return false;
-	}
-
-	while (ContentIndex != INDEX_NONE)
-	{
-		// Assume that the project directory is the parent of the /content/ directory
-		FString ProjectDirectory = PackagePath.Left(ContentIndex);
-		FString PluginDirectory;
-		
-		TArray<FString> ProjectFile;
-		TArray<FString> PluginFile;
-
-		IFileManager::Get().FindFiles(ProjectFile, *ProjectDirectory, TEXT(".uproject"));
-
-		if (ProjectFile.IsEmpty())
-		{
-			// If there was no project file, the package could be in a plugin, so lets check for that
-			PluginDirectory = ProjectDirectory;
-			IFileManager::Get().FindFiles(PluginFile, *PluginDirectory, TEXT(".uplugin"));
-
-			if (PluginFile.Num() == 1)
-			{
-				PluginFilePath = PluginDirectory / PluginFile[0];
-
-				// We have a valid plugin file, so we should be able to find a /plugins/ directory which will be just below the project directory
-				const int32 PluginIndex = PluginDirectory.Find(TEXT("/plugins/"), ESearchCase::IgnoreCase, ESearchDir::FromEnd);
-				if (PluginIndex != INDEX_NONE)
-				{
-					// We found the plugin root directory so the one above it should be the project directory
-					ProjectDirectory = PluginDirectory.Left(PluginIndex);
-					IFileManager::Get().FindFiles(ProjectFile, *ProjectDirectory, TEXT(".uproject"));
-				}
-			}
-			else if (PluginFile.Num() > 1)
-			{
-				UE_LOG(LogVirtualizationTool, Warning, TEXT("Found multiple .uplugin files for '%s' at '%s'"), *PackagePath, *PluginDirectory);
-				return false;
-			}
-		}
-
-		if (ProjectFile.Num() == 1)
-		{
-			ProjectFilePath = ProjectDirectory / ProjectFile[0];
-			return true;
-		}
-		else if (!ProjectFile.IsEmpty())
-		{
-			UE_LOG(LogVirtualizationTool, Warning, TEXT("Found multiple .uproject files for '%s' at '%s'"), *PackagePath, *ProjectDirectory);
-			return false;
-		}
-		
-		// Could be more than one content directory in the path so lets keep looking
-		ContentIndex = PackagePath.Find(TEXT("/content/"), ESearchCase::IgnoreCase, ESearchDir::FromEnd, ContentIndex);
-	}
-	
-	// We found one or more content directories but none of them contained a project file
-	UE_LOG(LogVirtualizationTool, Warning, TEXT("Failed to find project file for '%s'"), *PackagePath);
-	return false;
-}
-
 FProject& FUnrealVirtualizationToolApp::FindOrAddProject(FString&& ProjectFilePath)
 {
 	FProject* Project = Projects.FindByPredicate([&ProjectFilePath](const FProject& Project)->bool
