Remove unnecessary include

#rnx

diff --git a/Engine/Source/Developer/Virtualization/Private/VirtualizationManager.cpp b/Engine/Source/Developer/Virtualization/Private/VirtualizationManager.cpp
--- a/Engine/Source/Developer/Virtualization/Private/VirtualizationManager.cpp
+++ b/Engine/Source/Developer/Virtualization/Private/VirtualizationManager.cpp
@@ -5,7 +5,6 @@
 #include "AnalyticsEventAttribute.h"
 #include "HAL/IConsoleManager.h"
 #include "HAL/PlatformTime.h"
-#include "Logging/MessageLog.h"
 #include "Misc/App.h"
 #include "Misc/CommandLine.h"
 #include "Misc/ConfigCacheIni.h"
