Clean up deprecated code in the virtualization system.

#rnx

diff --git a/Engine/Source/Developer/Virtualization/Private/VirtualizationManager.cpp b/Engine/Source/Developer/Virtualization/Private/VirtualizationManager.cpp
--- a/Engine/Source/Developer/Virtualization/Private/VirtualizationManager.cpp
+++ b/Engine/Source/Developer/Virtualization/Private/VirtualizationManager.cpp
@@ -412,16 +412,6 @@ TArray<FString> SanitizeFilePaths(TConstArrayView<FString> FilePaths)
 	return SanitizedPaths;
 }
 
-/** 
- * Utility to make 'StorageType ==  EStorageType::Cache' checks easier while EStorageType::Local continues to exist.
- * When the deprecated value is removed this can also be removed and code calling it can just check for EStorageType::Cache */
-bool IsCacheType(EStorageType StorageType)
-{
-PRAGMA_DISABLE_DEPRECATION_WARNINGS
-	return StorageType == EStorageType::Local || StorageType ==  EStorageType::Cache;
-PRAGMA_ENABLE_DEPRECATION_WARNINGS
-}
-
 /**
  * Profiling data allowing us to track how payloads are being push/pulled during the lifespan of the process. Note that as all backends are
  * created at the same time, we don't need to add locked when accessing the maps. In addition FCookStats is thread safe when adding hits/misses
@@ -618,9 +608,6 @@ bool FVirtualizationManager::IsPushingEnabled(EStorageType StorageType) const
 
 	switch (StorageType)
 	{
-PRAGMA_DISABLE_DEPRECATION_WARNINGS
-		case EStorageType::Local:
-PRAGMA_ENABLE_DEPRECATION_WARNINGS
 		case EStorageType::Cache:
 			return !CacheStorageBackends.IsEmpty() && EnumHasAllFlags(CachingPolicy, ECachingPolicy::CacheOnPush);
 			break;
@@ -688,7 +675,7 @@ bool FVirtualizationManager::PushData(TArrayView<FPushRequest> Requests, EStorag
 		return false;
 	}
 
-	FBackendArray& Backends = IsCacheType(StorageType) ? CacheStorageBackends : PersistentStorageBackends;
+	FBackendArray& Backends = StorageType == EStorageType::Cache ? CacheStorageBackends : PersistentStorageBackends;
 
 	if (Backends.IsEmpty())
 	{
@@ -804,7 +791,7 @@ bool FVirtualizationManager::PushData(TArrayView<FPushRequest> Requests, EStorag
 
 	// For local storage we consider the push to have failed only if ALL backends gave an error, if at least one backend succeeded then the operation succeeded.
 	// For persistent storage we require that all backends succeeded, so any errors will fail the push operation.
-	return IsCacheType(StorageType) ? ErrorCount < Backends.Num() : ErrorCount == 0;
+	return StorageType == EStorageType::Cache ? ErrorCount < Backends.Num() : ErrorCount == 0;
 }
 
 bool FVirtualizationManager::PullData(TArrayView<FPullRequest> Requests)
@@ -873,7 +860,7 @@ EQueryResult FVirtualizationManager::QueryPayloadStatuses(TArrayView<const FIoHa
 		OutStatuses[Index] = Ids[Index].IsZero() ? EPayloadStatus::Invalid : EPayloadStatus::NotFound;
 	}
 
-	FBackendArray& Backends = IsCacheType(StorageType) ? CacheStorageBackends : PersistentStorageBackends;
+	FBackendArray& Backends = StorageType == EStorageType::Cache ? CacheStorageBackends : PersistentStorageBackends;
 
 	TArray<int8> HitCount;
 	TArray<bool> Results;
diff --git a/Engine/Source/Developer/Virtualization/Public/IVirtualizationBackend.h b/Engine/Source/Developer/Virtualization/Public/IVirtualizationBackend.h
--- a/Engine/Source/Developer/Virtualization/Public/IVirtualizationBackend.h
+++ b/Engine/Source/Developer/Virtualization/Public/IVirtualizationBackend.h
@@ -268,25 +268,6 @@ private:
 	/** Override to implement the backends connection code */
 	virtual EConnectionStatus OnConnect() = 0;
 
-private:
-	// Deprecated methods, will cause compiler errors if anyone has actually overridden IVirtualizationBackend
-	// pre 5.3 but there is currently no good compiler support for invoking compiler warnings when overriding
-	// a deprecated base method.
-
-	UE_DEPRECATED(5.2, "Use the overload that also takes EPushFlags")
-	virtual bool PushData(TArrayView<FPushRequest> Requests) final
-	{
-		checkNoEntry();
-		return false;
-	}
-
-	UE_DEPRECATED(5.3, "Use the overload that also takes EPullFlags")
-	virtual bool PullData(TArrayView<FPullRequest> Requests) final
-	{
-		checkNoEntry();
-		return false;
-	}
-
 private:
 
 	/** The name assigned to the backend by the virtualization graph */
diff --git a/Engine/Source/Runtime/Core/Public/Virtualization/VirtualizationSystem.h b/Engine/Source/Runtime/Core/Public/Virtualization/VirtualizationSystem.h
--- a/Engine/Source/Runtime/Core/Public/Virtualization/VirtualizationSystem.h
+++ b/Engine/Source/Runtime/Core/Public/Virtualization/VirtualizationSystem.h
@@ -698,13 +698,6 @@ public:
 		return PushData(MakeArrayView(&Request, 1), StorageType);
 	}
 
-	UE_DEPRECATED(5.1, "Use the overload of ::PushData(FIoHash, FCompressedBuffer, FString, EStorageType)")
-	bool PushData(const FIoHash& Id, const FCompressedBuffer& Payload, EStorageType StorageType, const FString& Context)
-	{
-		FPushRequest Request(Id, Payload, Context);
-		return PushData(MakeArrayView(&Request, 1), StorageType);
-	}
-
 	/**
 	 * Pull a number of payloads from the virtualization backends.
 	 *
@@ -758,23 +751,6 @@ public:
 	 * 			encountered in which case the contents of OutStatuses should be ignored.
 	 */
 	virtual EQueryResult QueryPayloadStatuses(TArrayView<const FIoHash> Ids, EStorageType StorageType, TArray<EPayloadStatus>& OutStatuses) = 0;
-
-	UE_DEPRECATED(5.1, "Call ::QueryPayloadStatuses instead")
-	bool DoPayloadsExist(TArrayView<const FIoHash> Ids, EStorageType StorageType, TArray<EPayloadStatus>& OutStatuses)
-	{
-		return QueryPayloadStatuses(Ids, StorageType, OutStatuses) != EQueryResult::Success;
-	}
-
-	UE_DEPRECATED(5.2, "Use The other overload that takes a EVirtualizationOptions parameter")
-	virtual EVirtualizationResult TryVirtualizePackages(TConstArrayView<FString> PackagePaths, TArray<FText>& OutDescriptionTags, TArray<FText>& OutErrors)
-	{
-		FVirtualizationResult Result = TryVirtualizePackages(PackagePaths, EVirtualizationOptions::None);
-
-		OutDescriptionTags = MoveTemp(Result.DescriptionTags);
-		OutErrors = MoveTemp(Result.Errors);
-
-		return Result.WasSuccessful() ? EVirtualizationResult::Success : EVirtualizationResult::Failed;
-	}
 	
 	/**
 	 * Runs the virtualization process on a set of packages. All of the packages will be parsed and any found to be containing locally stored
diff --git a/Engine/Source/Runtime/Core/Public/Virtualization/VirtualizationTypes.h b/Engine/Source/Runtime/Core/Public/Virtualization/VirtualizationTypes.h
--- a/Engine/Source/Runtime/Core/Public/Virtualization/VirtualizationTypes.h
+++ b/Engine/Source/Runtime/Core/Public/Virtualization/VirtualizationTypes.h
@@ -12,8 +12,6 @@ namespace UE::Virtualization
 /** Describes the type of storage to use for a given action */
 enum class EStorageType : int8
 {
-	/** Deprecated value, replaced by EStorageType::Cache */
-	Local UE_DEPRECATED(5.1, "Use EStorageType::Cache instead") = -1,
 	/** Store in the local cache backends, this can be called from any thread */
 	Cache = 0,
 	/** Store in the persistent backends, this can only be called from the game thread due to limitations with ISourceControlModule. */
