Make sure that duplicate slashes are removed from the path given to file system backends.

#rnx

diff --git a/Engine/Source/Developer/Virtualization/Private/VirtualizationFileBackend.cpp b/Engine/Source/Developer/Virtualization/Private/VirtualizationFileBackend.cpp
--- a/Engine/Source/Developer/Virtualization/Private/VirtualizationFileBackend.cpp
+++ b/Engine/Source/Developer/Virtualization/Private/VirtualizationFileBackend.cpp
@@ -35,7 +35,8 @@ bool FFileSystemBackend::Initialize(const FString& ConfigEntry)
 	}
 
 	FPaths::NormalizeDirectoryName(RootDirectory);
-
+	FPaths::RemoveDuplicateSlashes(RootDirectory);
+	
 	if (RootDirectory.IsEmpty())
 	{
 		UE_LOG(LogVirtualization, Error, TEXT("[%s] Config file entry 'Path=' was empty"), *GetDebugName());
