Demote some warnings to verbose logging when trying to identify the project file for a package during virtualization checks.

#rnx

- When the code was factored out to a common utility I kept the log warnings if the path doesn't conform what we expect (should have a content directory, should have a project file etc) but the calling code sohuld really be the one to decide if we should warn/error that a file is not under a valid project file.

diff --git a/Engine/Source/Developer/Virtualization/Private/VirtualizationUtilities.cpp b/Engine/Source/Developer/Virtualization/Private/VirtualizationUtilities.cpp
--- a/Engine/Source/Developer/Virtualization/Private/VirtualizationUtilities.cpp
+++ b/Engine/Source/Developer/Virtualization/Private/VirtualizationUtilities.cpp
@@ -197,7 +197,7 @@ bool TryFindProject(const FString& PackagePath, const FString& Extension, FStrin
 	// Early out if there is not a single content directory in the path
 	if (ContentIndex == INDEX_NONE)
 	{
-		UE_LOG(LogVirtualization, Warning, TEXT("'%s' is not under a content directory"), *PackagePath);
+		UE_LOG(LogVirtualization, Verbose, TEXT("'%s' is not under a content directory"), *PackagePath);
 		return false;
 	}
 
@@ -254,7 +254,7 @@ bool TryFindProject(const FString& PackagePath, const FString& Extension, FStrin
 	}
 
 	// We found one or more content directories but none of them contained a project file
-	UE_LOG(LogVirtualization, Warning, TEXT("Failed to find project file for '%s'"), *PackagePath);
+	UE_LOG(LogVirtualization, Verbose, TEXT("Failed to find project file for '%s'"), *PackagePath);
 
 	return false;
 }
