Improvements to VA analytics when first connecting to revision control.

#rnx

- We already have analytics events when the user fails to connect to their server, but we'd like to track how long people spend waiting for these connections to be made.
- Add a new analytics event "Editor.VA.SourceControlOnConnect" that we currently send when the connection is successfully made. At some point in the future we should merge all of the other analytics events in this method with this one.
- Add a slow task scope to FSourceControlBackend::OnConnect and force it to refresh. This should better inform the user that they are stuck waiting on their perforce server if the connection is unusually slow.

diff --git a/Engine/Source/Developer/Virtualization/Private/VirtualizationSourceControlBackend.cpp b/Engine/Source/Developer/Virtualization/Private/VirtualizationSourceControlBackend.cpp
--- a/Engine/Source/Developer/Virtualization/Private/VirtualizationSourceControlBackend.cpp
+++ b/Engine/Source/Developer/Virtualization/Private/VirtualizationSourceControlBackend.cpp
@@ -6,10 +6,10 @@
 #include "Containers/Ticker.h"
 #include "HAL/Event.h"
 #include "HAL/FileManager.h"
+#include "HAL/PlatformFileManager.h"
 #include "IO/IoHash.h"
 #include "ISourceControlModule.h"
 #include "ISourceControlProvider.h"
-#include "HAL/PlatformFileManager.h"
 #include "Interfaces/IPluginManager.h"
 #include "Logging/MessageLog.h"
 #include "Logging/StructuredLog.h"
@@ -19,6 +19,7 @@
 #include "Misc/PathViews.h"
 #include "Misc/Paths.h"
 #include "Misc/ScopeExit.h"
+#include "Misc/ScopedSlowTask.h"
 #include "SVirtualizationRevisionControlConnectionDialog.h"
 #include "SourceControlHelpers.h"
 #include "SourceControlInitSettings.h"
@@ -335,6 +336,10 @@ bool FSourceControlBackend::Initialize(const FString& ConfigEntry)
 
 IVirtualizationBackend::EConnectionStatus FSourceControlBackend::OnConnect()
 {
+	FScopedSlowTask SlowTask(1.0f, LOCTEXT("VARevControlConnection", "Connecting to virtualized asset revision control backend..."));
+	SlowTask.EnterProgressFrame(1.0f);
+	SlowTask.ForceRefresh();
+
 	// First check that the 'PerforceSourceControl' plugin exists and is enabled so that we can give specific
 	// error messages if not.
 	// We could try to force enable the plugin at this point, in the same way that the stand alone tool does
@@ -361,14 +366,24 @@ IVirtualizationBackend::EConnectionStatus FSourceControlBackend::OnConnect()
 	bool bSaveSettings = false; // Initially we will try to reason the perforce settings from the ini file and global environment
 								// so we don't want to save these values to the users ini files.
 
+	int32 Attempts = 0;
+	double TotalConnectionTime = 0.0;
+
 	while (true)
 	{
+		const double ConnectionStartTime = FPlatformTime::Seconds();
+
 		// TODO: At the moment we cannot get back error messages from ISourceControlProvider::Init, nor can we get back info about
 		// the actual login details used, so for the initial dialog we will just have to go with the settings we know and not 
 		// give the user any details about the connection problem.
+
 		FText ErrorMessage;
 		IVirtualizationBackend::EConnectionStatus ConnectionResult = OnConnectInternal(Port, UserName, bSaveSettings, ErrorMessage);
 
+		// Record each attempt and keep track of the total time we've spent connecting accross all attempts.
+		Attempts++;
+		TotalConnectionTime += FPlatformTime::Seconds() - ConnectionStartTime;
+
 		if (ConnectionResult == IVirtualizationBackend::EConnectionStatus::Connected)
 		{
 			if (bSaveSettings)
@@ -376,6 +391,19 @@ IVirtualizationBackend::EConnectionStatus FSourceControlBackend::OnConnect()
 				check(::IsInGameThread());
 				GConfig->Flush(false, SourceControlHelpers::GetSettingsIni());
 			}
+
+			{
+				TArray<FAnalyticsEventAttribute> Attributes;
+				Attributes.Reserve(4);
+
+				Attributes.Add({ TEXT("Connected"), true }); // Hard coded for now, but this will make more sense if we combibe this event with the failures below
+				Attributes.Add({ TEXT("Attempts"), Attempts });
+				Attributes.Add({ TEXT("TotalConnectionTime"), TotalConnectionTime });
+				Attributes.Add({ TEXT("Interactive"), Utils::IsProcessInteractive() });
+
+				GetAnalyticsRecordEvent().Broadcast(TEXT("Editor.VA.SourceControlOnConnect"), Attributes, EAnalyticsFlags::None);
+			}
+
 			return IVirtualizationBackend::EConnectionStatus::Connected;
 		}
 
