CIS: Comment out use of FScopedSlowTask in FSourceControlBackend::OnConnect as it is not safe in some scenarios (cooking commandlet at the very least)

#rnx

diff --git a/Engine/Source/Developer/Virtualization/Private/VirtualizationSourceControlBackend.cpp b/Engine/Source/Developer/Virtualization/Private/VirtualizationSourceControlBackend.cpp
--- a/Engine/Source/Developer/Virtualization/Private/VirtualizationSourceControlBackend.cpp
+++ b/Engine/Source/Developer/Virtualization/Private/VirtualizationSourceControlBackend.cpp
@@ -336,9 +336,9 @@ bool FSourceControlBackend::Initialize(const FString& ConfigEntry)
 
 IVirtualizationBackend::EConnectionStatus FSourceControlBackend::OnConnect()
 {
-	FScopedSlowTask SlowTask(1.0f, LOCTEXT("VARevControlConnection", "Connecting to virtualized asset revision control backend..."));
-	SlowTask.EnterProgressFrame(1.0f);
-	SlowTask.ForceRefresh();
+	//FScopedSlowTask SlowTask(1.0f, LOCTEXT("VARevControlConnection", "Connecting to virtualized asset revision control backend..."));
+	//SlowTask.EnterProgressFrame(1.0f);
+	//SlowTask.ForceRefresh();
 
 	// First check that the 'PerforceSourceControl' plugin exists and is enabled so that we can give specific
 	// error messages if not.
