Restore use of FSlowTask in FSourceControlBackend::OnConnect

#rnx

- This is just a backout of CL 39596397 which itself was a partial backout of 39584896 due to an assert when cooking.
- The assert was resolved in CL 39635250 which fixed FSlowTask::ForceRefresh to work correctly when called on a background thread.

diff --git a/Engine/Source/Developer/Virtualization/Private/VirtualizationSourceControlBackend.cpp b/Engine/Source/Developer/Virtualization/Private/VirtualizationSourceControlBackend.cpp
--- a/Engine/Source/Developer/Virtualization/Private/VirtualizationSourceControlBackend.cpp
+++ b/Engine/Source/Developer/Virtualization/Private/VirtualizationSourceControlBackend.cpp
@@ -336,9 +336,9 @@ bool FSourceControlBackend::Initialize(const FString& ConfigEntry)
 
 IVirtualizationBackend::EConnectionStatus FSourceControlBackend::OnConnect()
 {
-	//FScopedSlowTask SlowTask(1.0f, LOCTEXT("VARevControlConnection", "Connecting to virtualized asset revision control backend..."));
-	//SlowTask.EnterProgressFrame(1.0f);
-	//SlowTask.ForceRefresh();
+	FScopedSlowTask SlowTask(1.0f, LOCTEXT("OnConnect", "Connecting to virtualized asset revision control backend..."));
+	SlowTask.EnterProgressFrame(1.0f);
+	SlowTask.ForceRefresh();
 
 	// First check that the 'PerforceSourceControl' plugin exists and is enabled so that we can give specific
 	// error messages if not.
