Improve logging when hydrating virtualized assets.

#rnx

diff --git a/Engine/Source/Developer/Virtualization/Private/PackageRehydrationProcess.cpp b/Engine/Source/Developer/Virtualization/Private/PackageRehydrationProcess.cpp
--- a/Engine/Source/Developer/Virtualization/Private/PackageRehydrationProcess.cpp
+++ b/Engine/Source/Developer/Virtualization/Private/PackageRehydrationProcess.cpp
@@ -6,6 +6,7 @@
 #include "Internationalization/Internationalization.h"
 #include "Misc/PackageName.h"
 #include "Misc/PackagePath.h"
+#include "Misc/ScopeExit.h"
 #include "Misc/ScopedSlowTask.h"
 #include "PackageUtils.h"
 #include "Serialization/MemoryArchive.h"
@@ -192,6 +193,13 @@ void RehydratePackages(TConstArrayView<FString> PackagePaths, ERehydrationOption
 
 	const double StartTime = FPlatformTime::Seconds();
 
+	ON_SCOPE_EXIT
+	{
+		OutResult.TimeTaken = FPlatformTime::Seconds() - StartTime;
+		UE_LOG(LogVirtualization, Display, TEXT("Rehydration process complete"));
+		UE_LOG(LogVirtualization, Verbose, TEXT("Rehydration process took %.3f(s)"), OutResult.TimeTaken);
+	};
+
 	FScopedSlowTask Progress(1.0f, LOCTEXT("VAHydration_TaskOnDisk", "Re-hydrating Assets On Disk..."));
 	Progress.MakeDialog();
 
@@ -285,16 +293,31 @@ void RehydratePackages(TConstArrayView<FString> PackagePaths, ERehydrationOption
 		}
 	}
 
+	int64 OriginalSizeTotal = 0;
+	int64 HydratedSizeTotal = 0;
+
 	for (const TPair<FString, FString>& Pair : PackagesToReplace)
 	{
 		const FString& OriginalFilePath = Pair.Key;
-		const FString& TempFilePath = Pair.Value;
+		const FString& NewPackagePath = Pair.Value;
 
 		if (CanWriteToFile(OriginalFilePath))
 		{
-			if (IFileManager::Get().Move(*OriginalFilePath, *TempFilePath))
+			const int64 OriginalSize = IFileManager::Get().GetStatData(*OriginalFilePath).FileSize;
+			const int64 HydratedSize = IFileManager::Get().GetStatData(*NewPackagePath).FileSize;
+
+			if (IFileManager::Get().Move(*OriginalFilePath, *NewPackagePath))
 			{
 				OutResult.RehydratedPackages.Add(OriginalFilePath);
+
+				OriginalSizeTotal += OriginalSize;
+				HydratedSizeTotal += HydratedSize;
+
+				// Could consider path name here to match the virtualization process.
+				UE_LOG(LogVirtualization, Verbose, TEXT("Hydrating %s: %s -> %s"),
+					*OriginalFilePath,
+					*FText::AsMemory(OriginalSize).ToString(),
+					*FText::AsMemory(HydratedSize).ToString());
 			}
 			else
 			{
@@ -316,8 +339,15 @@ void RehydratePackages(TConstArrayView<FString> PackagePaths, ERehydrationOption
 		}
 	}
 
-	OutResult.TimeTaken = FPlatformTime::Seconds() - StartTime;
-	UE_LOG(LogVirtualization, Verbose, TEXT("Rehydration process took %.3f(s)"), OutResult.TimeTaken);
+	if (OriginalSizeTotal != HydratedSizeTotal)
+	{
+		UE_LOG(LogVirtualization, Display, TEXT("Total hydration increase %s (%s -> %s)"),
+			*FText::AsMemory(HydratedSizeTotal - OriginalSizeTotal).ToString(),
+			*FText::AsMemory(OriginalSizeTotal).ToString(),
+			*FText::AsMemory(HydratedSizeTotal).ToString());
+	}
+
+	UE_LOG(LogVirtualization, Display, TEXT("Hydrated %d package(s)"), OutResult.RehydratedPackages.Num());
 }
 
 void RehydratePackages(TConstArrayView<FString> PackagePaths, uint64 PaddingAlignment, TArray<FText>& OutErrors, TArray<FSharedBuffer>& OutPackages, TArray<FRehydrationInfo>* OutInfo)
