CIS: Fix UnrealVirtualizationTool compilation when compiled as a modular target.

#rnx

- Add dllexport to each function as the header is public.

diff --git a/Engine/Source/Developer/Virtualization/Public/VirtualizationUtilities.h b/Engine/Source/Developer/Virtualization/Public/VirtualizationUtilities.h
--- a/Engine/Source/Developer/Virtualization/Public/VirtualizationUtilities.h
+++ b/Engine/Source/Developer/Virtualization/Public/VirtualizationUtilities.h
@@ -9,6 +9,8 @@ class FPackagePath;
 
 struct FIoHash;
 
+#define UE_API VIRTUALIZATION_API
+
 namespace UE::Virtualization::Utils
 {
 
@@ -28,7 +30,7 @@ namespace UE::Virtualization::Utils
  *			The string builder must have a capacity of at least 52 characters 
 			to avoid reallocation.
  */
-void PayloadIdToPath(const FIoHash& Id, FStringBuilderBase& OutPath);
+UE_API void PayloadIdToPath(const FIoHash& Id, FStringBuilderBase& OutPath);
 
 /** 
  * Converts a FIoHash into a file path.
@@ -38,7 +40,7 @@ void PayloadIdToPath(const FIoHash& Id, FStringBuilderBase& OutPath);
  * @param	Id The payload identifier used to create the file path.
  * @return	The resulting file path.
  */
-FString PayloadIdToPath(const FIoHash& Id);
+UE_API FString PayloadIdToPath(const FIoHash& Id);
 
 /**
  * Fill in the given string builder with the human readable message of the current system
@@ -47,7 +49,7 @@ FString PayloadIdToPath(const FIoHash& Id);
  * able to be called and write that the error is unknown instead of assuming that the
  * operation was a success.
  */
-void GetFormattedSystemError(FStringBuilderBase& SystemErrorMessage);
+UE_API void GetFormattedSystemError(FStringBuilderBase& SystemErrorMessage);
 
 
 enum class ETrailerFailedReason : uint8
@@ -68,21 +70,23 @@ enum class ETrailerFailedReason : uint8
  * is in fact missing and assumes that the caller previously attempted to a call to
  * load a trailer from the package but it failed.
  */
-ETrailerFailedReason FindTrailerFailedReason(const FPackagePath& PackagePath);
+UE_API ETrailerFailedReason FindTrailerFailedReason(const FPackagePath& PackagePath);
 
 /** 
  * Parse the given path and expand any environment variables found with the format
  * $(EnvVarName). If the variable cannot be found the function will return false 
  * and OutExpandedPath will be reset.
  */
-bool ExpandEnvironmentVariables(FStringView InputPath, FStringBuilderBase& OutExpandedPath);
+UE_API bool ExpandEnvironmentVariables(FStringView InputPath, FStringBuilderBase& OutExpandedPath);
 
 /** 
  * Returns true if the process is interactive and we should display error dialogs to
  * the user.
  * */
-bool IsProcessInteractive();
+UE_API bool IsProcessInteractive();
 
-bool TryFindProject(const FString& PackagePath, const FString& Extension, FString& ProjectFilePath, FString& PluginFilePath);
+UE_API bool TryFindProject(const FString& PackagePath, const FString& Extension, FString& ProjectFilePath, FString& PluginFilePath);
 
 } // namespace UE::Virtualization::Utils
+
+#undef UE_API
