Bump VA legacy config options to be logged as errors so we can remove them in a future release.

#rnx

diff --git a/Engine/Source/Developer/Virtualization/Private/VirtualizationManager.cpp b/Engine/Source/Developer/Virtualization/Private/VirtualizationManager.cpp
--- a/Engine/Source/Developer/Virtualization/Private/VirtualizationManager.cpp
+++ b/Engine/Source/Developer/Virtualization/Private/VirtualizationManager.cpp
@@ -1227,13 +1227,13 @@ void FVirtualizationManager::ApplySettingsFromConfigFiles(const FConfigFile& Con
 		if (ConfigFile.GetBool(LegacyConfigSection, TEXT("FilterEngineContent"), bDummyValue) ||
 			ConfigFile.GetBool(ConfigSection, TEXT("FilterEngineContent"), bDummyValue))
 		{
-			UE_LOG(LogVirtualization, Warning, TEXT("\tFilterEngineContent is now deprecated (engine content is never virtualized)"));
+			UE_LOG(LogVirtualization, Error, TEXT("\tFilterEngineContent is now deprecated (engine content is never virtualized) please update before compatibility is removed in UE 5.9"));
 		}
 
 		if (ConfigFile.GetBool(LegacyConfigSection, TEXT("FilterEnginePluginContent"), bDummyValue) ||
 			ConfigFile.GetBool(ConfigSection, TEXT("FilterEnginePluginContent"), bDummyValue))
 		{
-			UE_LOG(LogVirtualization, Warning, TEXT("\tFilterEnginePluginContent is now deprecated (engine content is never virtualized)"));
+			UE_LOG(LogVirtualization, Error, TEXT("\tFilterEnginePluginContent is now deprecated (engine content is never virtualized) please update before compatibility is removed in UE 5.9"));
 		}
 	}
 
@@ -1258,10 +1258,10 @@ void FVirtualizationManager::ApplySettingsFromConfigFiles(const FConfigFile& Con
 
 		if (!LegacyEntries.IsEmpty())
 		{
-			UE_LOG(LogVirtualization, Warning, TEXT("\tFound %d legacy ini file settings under [Core.ContentVirtualization] that should be moved to [Core.VirtualizationModule]"), LegacyEntries.Num());
+			UE_LOG(LogVirtualization, Error, TEXT("\tFound %d legacy ini file settings under [Core.ContentVirtualization] that should be moved to [Core.VirtualizationModule], please update before compatibility is removed in UE 5.9"), LegacyEntries.Num());
 			for (const FString& LegacyEntry : LegacyEntries)
 			{
-				UE_LOG(LogVirtualization, Warning, TEXT("\t\t%s"), *LegacyEntry);
+				UE_LOG(LogVirtualization, Error, TEXT("\t\t%s"), *LegacyEntry);
 			}
 		}
 	}
