Added an overload version of  UE::Virtualization::Utils::TryFindProject which allows us to search for multiple project file extensions at the same time.

#rnx

- IFileManager::FindFiles only allows for a single extension, if we want to check multiple extensions we need to use IFileManager::IterateDirectory and use our own visitor to check each file found.
- Note that with the visitor we return the full path where as FindFiles was just returning the file name so there is no need to append the result to the original search directory anymore.

diff --git a/Engine/Source/Developer/Virtualization/Private/VirtualizationUtilities.cpp b/Engine/Source/Developer/Virtualization/Private/VirtualizationUtilities.cpp
--- a/Engine/Source/Developer/Virtualization/Private/VirtualizationUtilities.cpp
+++ b/Engine/Source/Developer/Virtualization/Private/VirtualizationUtilities.cpp
@@ -3,6 +3,7 @@
 #include "VirtualizationUtilities.h"
 #include "VirtualizationExperimentalUtilities.h"
 
+#include "GenericPlatform/GenericPlatformFile.h"
 #include "HAL/FileManager.h"
 #include "HAL/PlatformProcess.h"
 #include "IO/IoHash.h"
@@ -185,7 +186,43 @@ EPayloadFilterReason FixFilterFlags(FStringView PackagePath, uint64 SizeOnDisk,
 	return CurrentFilterFlags;
 }
 
-bool TryFindProject(const FString& PackagePath, const FString& Extension, FString& ProjectFilePath, FString& PluginFilePath)
+bool TryFindProject(const FString& PackagePath, const FString& ProjectExtension, FString& OutProjectFilePath, FString& OutPluginFilePath)
+{
+	return TryFindProject(PackagePath, MakeConstArrayView(&ProjectExtension, 1), OutProjectFilePath, OutPluginFilePath);
+}
+
+class FProjectFileVistitor : public IPlatformFile::FDirectoryVisitor
+{
+public:
+	FProjectFileVistitor(TConstArrayView<FString> InProjectExtensions, TArray<FString>& OutProjectFiles)
+		: ProjectExtensions(InProjectExtensions)
+		, ProjectFiles(OutProjectFiles)
+	{
+	}
+
+	virtual bool Visit(const TCHAR* FilenameOrDirectory, bool bIsDirectory) override
+	{
+		if (!bIsDirectory)
+		{
+			for (const FString& Extension : ProjectExtensions)
+			{
+				if (FStringView(FilenameOrDirectory).EndsWith(Extension))
+				{
+					ProjectFiles.Add(FilenameOrDirectory);
+					return true;
+				}
+			}
+		}
+
+		return true;
+	}
+
+private:
+	TConstArrayView<FString> ProjectExtensions;
+	TArray<FString>& ProjectFiles;
+};
+
+bool TryFindProject(const FString& PackagePath, TConstArrayView<FString> ProjectExtensions, FString& ProjectFilePath, FString& PluginFilePath)
 {
 	TRACE_CPUPROFILER_EVENT_SCOPE(UE::Virtualization::Utils::TryFindProject);
 
@@ -210,7 +247,9 @@ bool TryFindProject(const FString& PackagePath, const FString& Extension, FStrin
 		TArray<FString> ProjectFile;
 		TArray<FString> PluginFile;
 
-		IFileManager::Get().FindFiles(ProjectFile, *ProjectDirectory, *Extension);
+		FProjectFileVistitor ProjectFileVistitor(ProjectExtensions, ProjectFile);
+
+		IFileManager::Get().IterateDirectory(*ProjectDirectory, ProjectFileVistitor);
 
 		if (ProjectFile.IsEmpty())
 		{
@@ -228,7 +267,7 @@ bool TryFindProject(const FString& PackagePath, const FString& Extension, FStrin
 				{
 					// We found the plugin root directory so the one above it should be the project directory
 					ProjectDirectory = PluginDirectory.Left(PluginIndex);
-					IFileManager::Get().FindFiles(ProjectFile, *ProjectDirectory, *Extension);
+					IFileManager::Get().IterateDirectory(*ProjectDirectory, ProjectFileVistitor);
 				}
 			}
 			else if (PluginFile.Num() > 1)
@@ -240,7 +279,7 @@ bool TryFindProject(const FString& PackagePath, const FString& Extension, FStrin
 
 		if (ProjectFile.Num() == 1)
 		{
-			ProjectFilePath = ProjectDirectory / ProjectFile[0];
+			ProjectFilePath = ProjectFile[0];
 			return true;
 		}
 		else if (!ProjectFile.IsEmpty())
diff --git a/Engine/Source/Developer/Virtualization/Public/VirtualizationUtilities.h b/Engine/Source/Developer/Virtualization/Public/VirtualizationUtilities.h
--- a/Engine/Source/Developer/Virtualization/Public/VirtualizationUtilities.h
+++ b/Engine/Source/Developer/Virtualization/Public/VirtualizationUtilities.h
@@ -2,6 +2,7 @@
 
 #pragma once
 
+#include "Containers/ArrayView.h"
 #include "Containers/StringFwd.h"
 #include "Containers/UnrealString.h"
 
@@ -85,7 +86,31 @@ UE_API bool ExpandEnvironmentVariables(FStringView InputPath, FStringBuilderBase
  * */
 UE_API bool IsProcessInteractive();
 
-UE_API bool TryFindProject(const FString& PackagePath, const FString& Extension, FString& ProjectFilePath, FString& PluginFilePath);
+/**
+ * Find the project file for the project that contains the provided package file. Note that this only intended to work for package files as it makes
+ * several assumptions about the path conventions in order to speed up the search.
+ * 
+ * @PackagePath			Path of a package file to start searching from. The file must be stored under a /content/ directory to be considered valid.
+ * @ProjectExtension	The extension that the project file may have.
+ * @OutProjectFilePath	The project path will be written to this string if found.
+ * @OutPluginFilePath	If the package is stored in a plugin then the path to it's '.uplugin' file will be written to this string.
+ * 
+ * @return				True if a project file is found otherwise it will return false.
+ */
+UE_API bool TryFindProject(const FString& PackagePath, const FString& ProjectExtension, FString& OutProjectFilePath, FString& OutPluginFilePath);
+
+/**
+ * Find the project file for the project that contains the provided package file. Note that this only intended to work for package files as it makes
+ * several assumptions about the path conventions in order to speed up the search.
+ *
+ * @PackagePath			Path of a package file to start searching from. The file must be stored under a /content/ directory to be considered valid.
+ * @ProjectExtensions	An array of potential extensions that the project file may have.
+ * @OutProjectFilePath	The project path will be written to this string if found.
+ * @OutPluginFilePath	If the package is stored in a plugin then the path to it's '.uplugin' file will be written to this string.
+ *
+ * @return				True if a project file is found otherwise it will return false.
+ */
+UE_API bool TryFindProject(const FString& PackagePath, TConstArrayView<FString> ProjectExtensions, FString& OutProjectFilePath, FString& OutPluginFilePath);
 
 } // namespace UE::Virtualization::Utils
 
