IVirtualizationSystem::TryVirtualizePackages now returns more info about the virtualization process.

#rnx

- FVirtualizationResult now contains the total size of the packages before and after they were virtualized. Note that these totals only include the package files that were actually modified.

diff --git a/Engine/Source/Developer/Virtualization/Private/PackageVirtualizationProcess.cpp b/Engine/Source/Developer/Virtualization/Private/PackageVirtualizationProcess.cpp
--- a/Engine/Source/Developer/Virtualization/Private/PackageVirtualizationProcess.cpp
+++ b/Engine/Source/Developer/Virtualization/Private/PackageVirtualizationProcess.cpp
@@ -494,8 +494,8 @@ void VirtualizePackages(TConstArrayView<FString> PackagePaths, EVirtualizationOp
 
 		UE_LOG(LogVirtualization, Display, TEXT("Replacing old packages with the virtualized version..."));
 
-		int64 OriginalSizeTotal = 0;
-		int64 ReducedSizeTotal = 0;
+		OutResultInfo.PreOperationSize = 0;
+		OutResultInfo.PostOperationSize = 0;
 
 		for (const FPackageReplacement& PackageInfo : PackagesToReplace)
 		{
@@ -515,8 +515,8 @@ void VirtualizePackages(TConstArrayView<FString> PackagePaths, EVirtualizationOp
 			{
 				OutResultInfo.VirtualizedPackages.Add(OriginalPackagePath);
 
-				OriginalSizeTotal += OriginalSize;
-				ReducedSizeTotal += ReducedSize;
+				OutResultInfo.PreOperationSize += OriginalSize;
+				OutResultInfo.PostOperationSize += ReducedSize;
 
 				UE_LOG(LogVirtualization, Verbose, TEXT("Reducing %s: %s -> %s"),
 					*PackageInfo.Path.GetDebugName(),
@@ -533,12 +533,12 @@ void VirtualizePackages(TConstArrayView<FString> PackagePaths, EVirtualizationOp
 			}
 		}
 
-		if (OriginalSizeTotal != ReducedSizeTotal)
+		if (OutResultInfo.PreOperationSize != OutResultInfo.PostOperationSize)
 		{
 			UE_LOG(LogVirtualization, Display, TEXT("Total Reduction %s (%s -> %s)"),
-				*FText::AsMemory(OriginalSizeTotal - ReducedSizeTotal).ToString(),
-				*FText::AsMemory(OriginalSizeTotal).ToString(),
-				*FText::AsMemory(ReducedSizeTotal).ToString());
+				*FText::AsMemory(OutResultInfo.PreOperationSize - OutResultInfo.PostOperationSize).ToString(),
+				*FText::AsMemory(OutResultInfo.PreOperationSize).ToString(),
+				*FText::AsMemory(OutResultInfo.PostOperationSize).ToString());
 		}
 
 		UE_LOG(LogVirtualization, Display, TEXT("Replaced %d package(s)"), OutResultInfo.VirtualizedPackages.Num());
diff --git a/Engine/Source/Runtime/Core/Public/Virtualization/VirtualizationSystem.h b/Engine/Source/Runtime/Core/Public/Virtualization/VirtualizationSystem.h
--- a/Engine/Source/Runtime/Core/Public/Virtualization/VirtualizationSystem.h
+++ b/Engine/Source/Runtime/Core/Public/Virtualization/VirtualizationSystem.h
@@ -425,6 +425,11 @@ struct FVirtualizationResult
 	/** The length of time that the process took in seconds */
 	double TimeTaken = 0.0;
 
+	/** The total size of the packages that actually were virtualized before they were modified */
+	uint64 PreOperationSize = 0;
+	/** The total size of the packages that actually were virtualized after they have been modified */
+	uint64 PostOperationSize = 0;
+
 	/* Returns if the virtualization process succeeded or not */
 	bool WasSuccessful() const
 	{
