IVirtualizationSystem::TryRehydratePackages now returns more info about the rehydration process.

#rnx

- FRehydrationResult now contains the total size of the packages before and after they were rehydrated. Note that these totals only include the package files that were actually modified.

diff --git a/Engine/Source/Developer/Virtualization/Private/PackageRehydrationProcess.cpp b/Engine/Source/Developer/Virtualization/Private/PackageRehydrationProcess.cpp
--- a/Engine/Source/Developer/Virtualization/Private/PackageRehydrationProcess.cpp
+++ b/Engine/Source/Developer/Virtualization/Private/PackageRehydrationProcess.cpp
@@ -293,9 +293,6 @@ void RehydratePackages(TConstArrayView<FString> PackagePaths, ERehydrationOption
 		}
 	}
 
-	int64 OriginalSizeTotal = 0;
-	int64 HydratedSizeTotal = 0;
-
 	for (const TPair<FString, FString>& Pair : PackagesToReplace)
 	{
 		const FString& OriginalFilePath = Pair.Key;
@@ -310,8 +307,8 @@ void RehydratePackages(TConstArrayView<FString> PackagePaths, ERehydrationOption
 			{
 				OutResult.RehydratedPackages.Add(OriginalFilePath);
 
-				OriginalSizeTotal += OriginalSize;
-				HydratedSizeTotal += HydratedSize;
+				OutResult.PreOperationSize += OriginalSize;
+				OutResult.PostOperationSize += HydratedSize;
 
 				// Could consider path name here to match the virtualization process.
 				UE_LOG(LogVirtualization, Verbose, TEXT("Hydrating %s: %s -> %s"),
@@ -339,12 +336,12 @@ void RehydratePackages(TConstArrayView<FString> PackagePaths, ERehydrationOption
 		}
 	}
 
-	if (OriginalSizeTotal != HydratedSizeTotal)
+	if (OutResult.PreOperationSize != OutResult.PostOperationSize)
 	{
 		UE_LOG(LogVirtualization, Display, TEXT("Total hydration increase %s (%s -> %s)"),
-			*FText::AsMemory(HydratedSizeTotal - OriginalSizeTotal).ToString(),
-			*FText::AsMemory(OriginalSizeTotal).ToString(),
-			*FText::AsMemory(HydratedSizeTotal).ToString());
+			*FText::AsMemory(OutResult.PostOperationSize - OutResult.PreOperationSize).ToString(),
+			*FText::AsMemory(OutResult.PreOperationSize).ToString(),
+			*FText::AsMemory(OutResult.PostOperationSize).ToString());
 	}
 
 	UE_LOG(LogVirtualization, Display, TEXT("Hydrated %d package(s)"), OutResult.RehydratedPackages.Num());
diff --git a/Engine/Source/Runtime/Core/Public/Virtualization/VirtualizationSystem.h b/Engine/Source/Runtime/Core/Public/Virtualization/VirtualizationSystem.h
--- a/Engine/Source/Runtime/Core/Public/Virtualization/VirtualizationSystem.h
+++ b/Engine/Source/Runtime/Core/Public/Virtualization/VirtualizationSystem.h
@@ -470,6 +470,11 @@ struct FRehydrationResult
 	/** The length of time that the process took in seconds */
 	double TimeTaken = 0.0;
 
+	/** The total size of the packages that actually were rehydrated before they were modified */
+	uint64 PreOperationSize = 0;
+	/** The total size of the packages that actually were rehydrated after they have been modified */
+	uint64 PostOperationSize = 0;
+
 	/* Returns if the rehydration process succeeded or not */
 	bool WasSuccessful() const
 	{
