Add missing cpu trace scope

#rnx

diff --git a/Engine/Source/Developer/Virtualization/Private/VirtualizationSourceControlBackend.cpp b/Engine/Source/Developer/Virtualization/Private/VirtualizationSourceControlBackend.cpp
--- a/Engine/Source/Developer/Virtualization/Private/VirtualizationSourceControlBackend.cpp
+++ b/Engine/Source/Developer/Virtualization/Private/VirtualizationSourceControlBackend.cpp
@@ -336,6 +336,8 @@ bool FSourceControlBackend::Initialize(const FString& ConfigEntry)
 
 IVirtualizationBackend::EConnectionStatus FSourceControlBackend::OnConnect()
 {
+	TRACE_CPUPROFILER_EVENT_SCOPE(FSourceControlBackend::OnConnect);
+
 	FScopedSlowTask SlowTask(1.0f, LOCTEXT("OnConnect", "Connecting to virtualized asset revision control backend..."));
 	SlowTask.EnterProgressFrame(1.0f);
 	SlowTask.ForceRefresh();
