Fix some typos in the code documentation.

#rnx

#Synced-CL 51757651

diff --git a/Engine/Source/Developer/Virtualization/Private/VirtualizationSourceControlBackend.cpp b/Engine/Source/Developer/Virtualization/Private/VirtualizationSourceControlBackend.cpp
--- a/Engine/Source/Developer/Virtualization/Private/VirtualizationSourceControlBackend.cpp
+++ b/Engine/Source/Developer/Virtualization/Private/VirtualizationSourceControlBackend.cpp
@@ -382,7 +382,7 @@ IVirtualizationBackend::EConnectionStatus FSourceControlBackend::OnConnect()
 		FText ErrorMessage;
 		IVirtualizationBackend::EConnectionStatus ConnectionResult = OnConnectInternal(Port, UserName, bSaveSettings, ErrorMessage);
 
-		// Record each attempt and keep track of the total time we've spent connecting accross all attempts.
+		// Record each attempt and keep track of the total time we've spent connecting across all attempts.
 		Attempts++;
 		TotalConnectionTime += FPlatformTime::Seconds() - ConnectionStartTime;
 
@@ -398,7 +398,7 @@ IVirtualizationBackend::EConnectionStatus FSourceControlBackend::OnConnect()
 				TArray<FAnalyticsEventAttribute> Attributes;
 				Attributes.Reserve(4);
 
-				Attributes.Add({ TEXT("Connected"), true }); // Hard coded for now, but this will make more sense if we combibe this event with the failures below
+				Attributes.Add({ TEXT("Connected"), true }); // Hard coded for now, but this will make more sense if we combine this event with the failures below
 				Attributes.Add({ TEXT("Attempts"), Attempts });
 				Attributes.Add({ TEXT("TotalConnectionTime"), TotalConnectionTime });
 				Attributes.Add({ TEXT("Interactive"), Utils::IsProcessInteractive() });
@@ -837,7 +837,7 @@ bool FSourceControlBackend::PushData(TArrayView<FPushRequest> Requests, EPushFla
 			CreateWorkspaceCommand->SetType(FCreateWorkspace::EType::Partitioned);
 		}
 
-		CreateWorkspaceCommand->SetDescription(TEXT("This workspace was autogenerated when submitting virtualized payloads to revision control"));
+		CreateWorkspaceCommand->SetDescription(TEXT("This workspace was auto generated when submitting virtualized payloads to revision control"));
 
 		if (SCCProvider->Execute(CreateWorkspaceCommand) != ECommandResult::Succeeded)
 		{
@@ -1134,7 +1134,7 @@ bool FSourceControlBackend::TryApplySettingsFromConfigFiles(const FString& Confi
 			RetryWaitTimeMS = RetryWaitTimeMSIniFile;
 		}
 
-		UE_LOG(LogVirtualization, Log, TEXT("[%s] Will retry failed downloads attempts %d time(s) with a gap of %dms betwen them"), *GetDebugName(), RetryCount, RetryWaitTimeMS);
+		UE_LOG(LogVirtualization, Log, TEXT("[%s] Will retry failed downloads attempts %d time(s) with a gap of %dms between them"), *GetDebugName(), RetryCount, RetryWaitTimeMS);
 	}
 
 	// Allow the number of concurrent connections to be limited
@@ -1145,7 +1145,7 @@ bool FSourceControlBackend::TryApplySettingsFromConfigFiles(const FString& Confi
 		if (MaxLimit != INDEX_NONE)
 		{
 			ConcurrentConnectionLimit = MakeUnique<FSemaphore>(MaxLimit, FSemaphore::EFlags::PrioritizeGameThread);
-			UE_LOG(LogVirtualization, Log, TEXT("[%s] Limted to %d concurrent revision control connections"), *GetDebugName(), MaxLimit);
+			UE_LOG(LogVirtualization, Log, TEXT("[%s] Limited to %d concurrent revision control connections"), *GetDebugName(), MaxLimit);
 		}
 		else
 		{
@@ -1170,7 +1170,7 @@ bool FSourceControlBackend::TryApplySettingsFromConfigFiles(const FString& Confi
 
 		if (bSuppressNotifications)
 		{
-			UE_LOG(LogVirtualization, Log, TEXT("[%s] Connection pop up warnings are supressed"), *GetDebugName());
+			UE_LOG(LogVirtualization, Log, TEXT("[%s] Connection pop up warnings are suppressed"), *GetDebugName());
 		}
 		else
 		{
