Update the documentation and code default for 'UseLegacyErrorHandling' to reflect the default set by the config file settings.

#rnx

- Back in CL 25162842 the default value for UseLegacyErrorHandling was changed in Main/Engine/Config/BaseEngine.ini from true to false, enabling the new error handling path by default.
- This change updates the cpp code to reflect that new default, without this change it looks like the legacy path is enabled by default.

#Synced-CL 51757651

diff --git a/Engine/Source/Developer/Virtualization/Private/VirtualizationManager.cpp b/Engine/Source/Developer/Virtualization/Private/VirtualizationManager.cpp
--- a/Engine/Source/Developer/Virtualization/Private/VirtualizationManager.cpp
+++ b/Engine/Source/Developer/Virtualization/Private/VirtualizationManager.cpp
@@ -529,7 +529,7 @@ FVirtualizationManager::FVirtualizationManager()
 	, bFilterMapContent(true)
 	, bAllowSubmitIfVirtualizationFailed(false)
 	, bLazyInitConnections(false)
-	, bUseLegacyErrorHandling(true)
+	, bUseLegacyErrorHandling(false)
 	, bForceCachingOnPull(false)
 	, bAllowBackendReconnection(true)
 {
diff --git a/Engine/Source/Developer/Virtualization/Private/VirtualizationManager.h b/Engine/Source/Developer/Virtualization/Private/VirtualizationManager.h
--- a/Engine/Source/Developer/Virtualization/Private/VirtualizationManager.h
+++ b/Engine/Source/Developer/Virtualization/Private/VirtualizationManager.h
@@ -125,7 +125,7 @@ struct FAnalyticsEventAttribute;
  * UseLegacyErrorHandling [bool]:				Controls how we deal with errors encountered when pulling payloads. When true a failed payload 
  *												pull will return an error and allow the process to carry on (the original error handling logic)
  *												and when false a dialog will be displayed to the user warning them about the failed pull and 
- *												prompting them to retry the pull or to quit the process. [Default=true]
+ *												prompting them to retry the pull or to quit the process. [Default=false]
  * PullErrorAdditionalMsg [string]				An additional message that will be added to the error dialog presented on payload pull failure.
  *												This allows you to add custom information, such as links to internal help docs without editing
  *												code. Note that this additional message only works with the error dialog and will do nothing
