Fix some additional minor typos in the Virtualization module.

#rnx

#Synced-CL 51757651

diff --git a/Engine/Source/Developer/Virtualization/Private/VirtualizationFileBackend.cpp b/Engine/Source/Developer/Virtualization/Private/VirtualizationFileBackend.cpp
--- a/Engine/Source/Developer/Virtualization/Private/VirtualizationFileBackend.cpp
+++ b/Engine/Source/Developer/Virtualization/Private/VirtualizationFileBackend.cpp
@@ -59,7 +59,7 @@ bool FFileSystemBackend::Initialize(const FString& ConfigEntry)
 
 	// Now log a summary of the backend settings to make issues easier to diagnose
 	UE_LOG(LogVirtualization, Log, TEXT("[%s] Using path: '%s'"), *GetDebugName(), *FPaths::ConvertRelativePathToFull(RootDirectory));
-	UE_LOG(LogVirtualization, Log, TEXT("[%s] Will retry failed read attempts %d times with a gap of %dms betwen them"), *GetDebugName(), RetryCount, RetryWaitTimeMS);
+	UE_LOG(LogVirtualization, Log, TEXT("[%s] Will retry failed read attempts %d times with a gap of %dms between them"), *GetDebugName(), RetryCount, RetryWaitTimeMS);
 
 	return true;
 }
diff --git a/Engine/Source/Developer/Virtualization/Private/VirtualizationManager.cpp b/Engine/Source/Developer/Virtualization/Private/VirtualizationManager.cpp
--- a/Engine/Source/Developer/Virtualization/Private/VirtualizationManager.cpp
+++ b/Engine/Source/Developer/Virtualization/Private/VirtualizationManager.cpp
@@ -774,7 +774,7 @@ bool FVirtualizationManager::PushData(TArrayView<FPushRequest> Requests, EStorag
 		}
 	}
 
-	UE_CLOG(ErrorCount == Backends.Num(), LogVirtualization, Error, TEXT("Failed to push '%d' payload(s) to any backend'"), ValidatedRequests.Num());
+	UE_CLOG(ErrorCount == Backends.Num(), LogVirtualization, Error, TEXT("Failed to push '%d' payload(s) to any backend"), ValidatedRequests.Num());
 
 	// Now we need to update the statuses of the original list of requests with those from our validated list
 	for (int32 Index = 0; Index < Requests.Num(); ++Index)
@@ -822,7 +822,7 @@ bool FVirtualizationManager::PullData(TArrayView<FPullRequest> Requests)
 			// code handles it properly.
 			// Could be worth extending ::PullData to return error codes instead so we can make a better distinction 
 			// between the payload not being found in any of the backends and one or more of the backends failing.
-			UE_LOG(LogVirtualization, Error, TEXT("Payload '%s' failed to be pulled from any backend'"), *LexToString(Request.GetIdentifier()));
+			UE_LOG(LogVirtualization, Error, TEXT("Payload '%s' failed to be pulled from any backend"), *LexToString(Request.GetIdentifier()));
 
 			bSuccess = false;
 		}
@@ -1308,7 +1308,7 @@ void FVirtualizationManager::ApplyDebugSettingsFromFromCmdline()
 	FString CmdlineGraphName;
 	if (IsCmdLineValueSet(TEXT("-VABackendGraph="), TEXT("-VA-BackendGraph="), CmdlineGraphName))
 	{
-		UE_LOG(LogVirtualization, Display, TEXT("Backend graph overriden from the cmdline: '%s'"), *CmdlineGraphName);
+		UE_LOG(LogVirtualization, Display, TEXT("Backend graph overridden from the cmdline: '%s'"), *CmdlineGraphName);
 		BackendGraphName = CmdlineGraphName;
 	}
 
@@ -1403,7 +1403,7 @@ void FVirtualizationManager::RegisterConsoleCommands()
 
 	DebugValues.ConsoleObjects.Add(IConsoleManager::Get().RegisterConsoleCommand(
 		TEXT("VA.MissBackends"),
-		TEXT("A debug commnad which can be used to disable payload pulling on one or more backends"),
+		TEXT("A debug command which can be used to disable payload pulling on one or more backends"),
 		FConsoleCommandWithArgsAndOutputDeviceDelegate::CreateRaw(this, &FVirtualizationManager::OnUpdateDebugMissBackendsFromConsole)
 	));
 
@@ -2328,7 +2328,7 @@ void FVirtualizationManager::GatherAnalytics(TArray<FAnalyticsEventAttribute>& A
 
 EPayloadFilterReason FVirtualizationManager::FixFilterFlags(FStringView PackagePath, uint64 SizeOnDisk, EPayloadFilterReason CurrentFilterFlags)
 {
-	// We only apply new filters if the payload does not curently have any as stored filtered reasons would be checked first before
+	// We only apply new filters if the payload does not currently have any as stored filtered reasons would be checked first before
 	// we ever got to this part of the virtualization process.
 	if (CurrentFilterFlags == EPayloadFilterReason::None)
 	{
